Async Message System Plugin.

This provides an API to pass messages in between threads and tick groups to make it easier to coordinate and synchronize async work in Unreal Engine.

The general flow of using this system is that you have some data on a different thread which you would like publish to another thread. On the thread where you have your data avialable, you queue an event for broadcasting. Where you would like to receive message that data specifies it's binding options (FAsyncMessageBindingOptions) for *when* it is safe for that event to receive that data. Any time an event is queued, the payload data is copied to make it safe for listeners to read on different threads. If you have large payloads, this could be costly, so be careful.

The system will bind listeners via weak references, and if the reference goes out of scope, the listener will be unbound.

Messages have a hierarchy based on their ID, which is a gameplay tag. You can listen for any "parent" of specific events, which is very useful for handling large groups of events.

The main useage of this will be through the UAsyncMessageWorldSubsystem, which has a FAsyncGameplayMessageSystem on it for you to use. You may also find it useful to create your own message system with different ways of scheudling by subclassing the FAsyncMessageSystemBase class, if you have a more customized threading model for your game.

At the time of writing, this system is only supported in C++ but a Blueprint interface is planned for the future.

See AsyncMessageSystemTests.cpp for examples of useage.

diff --git a/Engine/Plugins/Experimental/AsyncMessageSystem/AsyncMessageSystem.uplugin b/Engine/Plugins/Experimental/AsyncMessageSystem/AsyncMessageSystem.uplugin
new file mode 100644
--- /dev/null
+++ b/Engine/Plugins/Experimental/AsyncMessageSystem/AsyncMessageSystem.uplugin
@@ -0,0 +1,24 @@
+{
+	"FileVersion": 3,
+	"Version": 1,
+	"VersionName": "0.1",
+	"Description": "Async Message System",
+	"Category": "Framework",
+	"CreatedBy": "Epic Games, Inc.",
+	"CreatedByURL": "https://epicgames.com",
+	"DocsURL": "",
+	"MarketplaceURL": "",
+	"SupportURL": "",
+	"EnabledByDefault": false,
+	"CanContainContent": false,
+	"IsBetaVersion": false,
+	"IsExperimentalVersion": true,
+	"Installed": false,
+	"Modules": [
+		{
+			"Name": "AsyncMessageSystem",
+			"Type": "Runtime",
+			"LoadingPhase": "PreDefault"
+		}
+	]
+}
\ No newline at end of file
diff --git a/Engine/Plugins/Experimental/AsyncMessageSystem/Source/AsyncMessageSystem/AsyncMessageSystem.build.cs b/Engine/Plugins/Experimental/AsyncMessageSystem/Source/AsyncMessageSystem/AsyncMessageSystem.build.cs
new file mode 100644
--- /dev/null
+++ b/Engine/Plugins/Experimental/AsyncMessageSystem/Source/AsyncMessageSystem/AsyncMessageSystem.build.cs
@@ -0,0 +1,21 @@
+// Copyright Epic Games, Inc. All Rights Reserved.
+
+namespace UnrealBuildTool.Rules
+{
+	public class AsyncMessageSystem : ModuleRules
+	{
+		public AsyncMessageSystem(ReadOnlyTargetRules Target) : base(Target)
+		{
+			PCHUsage = ModuleRules.PCHUsageMode.UseExplicitOrSharedPCHs;
+			
+			PublicDependencyModuleNames.AddRange(
+				new string[] {
+					"Core",
+					"CoreUObject",
+					"Engine",
+					"GameplayTags",
+				}
+			);
+		}
+	}
+}
\ No newline at end of file
diff --git a/Engine/Plugins/Experimental/AsyncMessageSystem/Source/AsyncMessageSystem/Private/AsyncGameplayMessageSystem.cpp b/Engine/Plugins/Experimental/AsyncMessageSystem/Source/AsyncMessageSystem/Private/AsyncGameplayMessageSystem.cpp
new file mode 100644
--- /dev/null
+++ b/Engine/Plugins/Experimental/AsyncMessageSystem/Source/AsyncMessageSystem/Private/AsyncGameplayMessageSystem.cpp
@@ -0,0 +1,224 @@
+// Copyright Epic Games, Inc. All Rights Reserved.
+
+#include "AsyncGameplayMessageSystem.h"
+#include "AsyncMessageSystemLogs.h"
+#include "Async/Async.h"
+#include "Engine/Level.h"
+#include "Engine/World.h"
+
+namespace UE::Private
+{
+	static FString LexToString(const ETickingGroup Group)
+	{
+		static FString Invalid = TEXT("Invalid");
+		static UEnum* TGEnum = StaticEnum<ETickingGroup>();
+		
+		return TGEnum ? TGEnum->GetDisplayNameTextByValue(Group).ToString() : Invalid;
+	};
+}
+
+/**
+ * Tick function which will being the processing of messages for specific tick groups on the message system. 
+ */
+struct FMessageSystemTickFunction final : public FTickFunction
+{
+	friend class FAsyncGameplayMessageSystem;
+	
+	explicit FMessageSystemTickFunction(const ETickingGroup InGroup, TWeakPtr<FAsyncGameplayMessageSystem> InWeakMessageSystem)
+		: FTickFunction()
+		, WeakMessageSys(InWeakMessageSystem)
+	{
+		bCanEverTick = true;
+		bStartWithTickEnabled = true;
+		bAllowTickBatching = true;
+
+		// Only run this tick function on the game thread, because our message system is our "sync point"
+		// for everything. 
+		bRunOnAnyThread = false;
+
+		// We want to ensure that we start and end in the same tick group to make sure we are a valid sync point for other threads
+		TickGroup = InGroup;
+		EndTickGroup = TickGroup;
+	}
+
+	// FTickFunction interface
+	virtual void ExecuteTick(float DeltaTime, ELevelTick TickType, ENamedThreads::Type CurrentThread, const FGraphEventRef& MyCompletionGraphEvent) override
+	{
+		// Each tick function can simply call the message system and let it know that the next tick group has started.
+		if (TSharedPtr<FAsyncGameplayMessageSystem> MessageSys = WeakMessageSys.Pin())
+		{
+			MessageSys->ExecuteTick(DeltaTime, TickType, CurrentThread, MyCompletionGraphEvent, TickGroup);
+		}
+	}
+	
+	virtual FString DiagnosticMessage() override
+	{
+		return TEXT("FMessageSystemTickFunction::") + UE::Private::LexToString(TickGroup);
+	}
+
+	/**
+	 * The owning message system which this tick function is going to update
+	 */
+	TWeakPtr<FAsyncGameplayMessageSystem> WeakMessageSys = nullptr;
+};
+
+FAsyncGameplayMessageSystem::FAsyncGameplayMessageSystem(UWorld* OwningWorld)
+	: OuterWorld(MakeWeakObjectPtr<UWorld>(OwningWorld))
+{
+	
+}
+
+void FAsyncGameplayMessageSystem::Startup_Impl()
+{
+	check(OuterWorld.IsValid());
+	
+	// Create a tick function for each tick group
+	CreateTickFunctions();
+}
+
+void FAsyncGameplayMessageSystem::Shutdown_Impl()
+{
+	// Remove all tick groups and wait for them to finish
+	DestroyTickFunctions();
+}
+
+void FAsyncGameplayMessageSystem::PostQueueMessage(const FAsyncMessageId MessageId, const TSet<FAsyncMessageBindingOptions>& OptionsBoundTo)
+{
+	TRACE_CPUPROFILER_EVENT_SCOPE(FMessageSystem::PostQueueMessage);
+
+	// When we queue a message, check if there are any listeners outside of tick groups who would need a specific
+	// async task to process their message queue. 
+	for (const FAsyncMessageBindingOptions& BindingOpts : OptionsBoundTo)
+	{
+		if (BindingOpts.Type == FAsyncMessageBindingOptions::EBindingType::UseTaskId)
+		{
+			StartAsyncProcessForBinding(BindingOpts);
+		}
+
+		// TODO: We could potentially use this to conditionally control when the tick groups need to run as well, reducing the amount that they are called.
+		// That way, the tick functions only execute if there are currently messages in their queue. We could possibly do this by registering/unregistering the
+		// tick functions as needed for the tick groups as they run. For now though, just let them tick
+	}
+}
+
+void FAsyncGameplayMessageSystem::CreateTickFunctions()
+{
+	check(TickFunctions.IsEmpty());
+
+	if (!ensureMsgf(OuterWorld.IsValid(), TEXT("Failed to create message system tick functions, the outer world is invalid!")))
+	{
+		return;
+	}
+
+	// We will be binding to some owning world's persistent level to create our tick functions
+	ULevel* TickLevel = OuterWorld->PersistentLevel;
+
+	TWeakPtr<FAsyncGameplayMessageSystem> WeakThisPtr = StaticCastWeakPtr<FAsyncGameplayMessageSystem>(this->AsWeak());
+
+	// Track the previous tick function so that we can add it as a dependency for each to finish
+	TSharedPtr<FTickFunction> PreviousTickFunction = nullptr;
+
+	const int StartingGroup = static_cast<int>(EarliestSupportedTickGroup);
+	const int LastGroup = static_cast<int>(LatestSupportedTickGroup);
+
+	// Spawn a tick function for every tick group that we can actually do any work in
+	for (int i = StartingGroup; i <= LastGroup; ++i)
+	{
+		const ETickingGroup Group = static_cast<ETickingGroup>(i);
+
+		TSharedPtr<FTickFunction> Func = MakeShared<FMessageSystemTickFunction>(Group, WeakThisPtr);
+
+		Func->RegisterTickFunction(TickLevel);
+
+		// We always want the previous tick function to wait until the next one to start processing
+		if (PreviousTickFunction.IsValid())
+		{
+			Func->AddPrerequisite(TickLevel, *PreviousTickFunction);
+		}
+
+		PreviousTickFunction = Func;
+
+		// Keep track of the tick functions we have created so that we can properly unregister them later
+		TickFunctions.Emplace(MoveTemp(Func));
+	}
+}
+
+void FAsyncGameplayMessageSystem::DestroyTickFunctions()
+{
+	for (TSharedPtr<FTickFunction> Func : TickFunctions)
+	{
+		if (Func.IsValid())
+		{
+			Func->UnRegisterTickFunction();
+		}
+	}
+
+	TickFunctions.Empty();
+}
+
+void FAsyncGameplayMessageSystem::StartAsyncProcessForBinding(const FAsyncMessageBindingOptions& Options)
+{
+	TRACE_CPUPROFILER_EVENT_SCOPE(FMessageSystem::StartAsyncProcessForBinding);
+
+	// We only want to do this for task ID's, tick groups are already being processed by us.
+	check(Options.Type == FAsyncMessageBindingOptions::EBindingType::UseTaskId);
+
+	// Kick off a weak lambda to process the message queue for this task graph ID, ensuring
+	// that any listeners bound to these options will get the message called back when they expect
+	TWeakPtr<FAsyncGameplayMessageSystem> WeakThisPtr = StaticCastWeakPtr<FAsyncGameplayMessageSystem>(this->AsWeak());
+
+	
+	const ENamedThreads::Type ThreadToProcessOn = static_cast<ENamedThreads::Type>(Options.TaskId); 
+
+		// TODO: This could be a UE::Tasks::Launch call, but we would need to convert the ENamedThreads to ETaskPriority and EExtendedTaskPriority.
+    	// This conversion only exists in one direction right now and is not public, see ENamedThreads::Type TranslatePriority.
+	AsyncTask(ThreadToProcessOn, [WeakThisPtr, Options]()
+	{
+		if (TSharedPtr<FAsyncGameplayMessageSystem> This = WeakThisPtr.Pin())
+		{
+			This->ProcessMessageQueueForBinding(Options);
+		}
+	});
+}
+
+void FAsyncGameplayMessageSystem::ExecuteTick(
+	float DeltaTime,
+	ELevelTick TickType,
+	ENamedThreads::Type CurrentThread,
+	const FGraphEventRef& MyCompletionGraphEvent,
+	ETickingGroup TickGroup)
+{
+	TRACE_CPUPROFILER_EVENT_SCOPE(FMessageSystem::ExecuteTick);
+
+	if (!OuterWorld.IsValid())
+	{
+		UE_LOG(LogAsyncMessageSystem, Warning, TEXT("[%hs] OuterWorld weak pointer is no longer valid for message system. Messages will not be processed, and this system will be shut down."), __func__);
+		Shutdown();
+		return;
+	}
+	
+	CurrentTickGroup = TickGroup;
+
+	// TODO: Right now the message system uses GFrameCounter as a quick and easy way to determine when messages are sent
+	// We should get away from that, and instead use an atomic in in the message system itself and increment it once per frame
+	
+	// If we are in the last tick group of the frame, then increment the frame counter of the message system that it is using.
+	//const bool bIsLastGroupOfFrame = (CurrentTickGroup == LatestSupportedTickGroup);
+	const uint64 CurrentFrame = GFrameCounter;
+	
+	UE_LOG(LogAsyncMessageSystem, VeryVerbose, TEXT("Frame %llu Exectute tick %s :: Last tick group was %s"),
+		CurrentFrame,
+		*UE::Private::LexToString(CurrentTickGroup),
+		*UE::Private::LexToString(LastTickedGroup));
+
+	ensure(CurrentTickGroup >= EarliestSupportedTickGroup || CurrentTickGroup <= LatestSupportedTickGroup);	
+
+	// Process the messages for this current tick group
+	FAsyncMessageBindingOptions Options = {};
+	Options.Type = FAsyncMessageBindingOptions::EBindingType::UseTickGroup;
+	Options.Group = CurrentTickGroup;
+	
+	ProcessMessageQueueForBinding(Options);
+		
+	LastTickedGroup = TickGroup;
+}
\ No newline at end of file
diff --git a/Engine/Plugins/Experimental/AsyncMessageSystem/Source/AsyncMessageSystem/Private/AsyncMessage.cpp b/Engine/Plugins/Experimental/AsyncMessageSystem/Source/AsyncMessageSystem/Private/AsyncMessage.cpp
new file mode 100644
--- /dev/null
+++ b/Engine/Plugins/Experimental/AsyncMessageSystem/Source/AsyncMessageSystem/Private/AsyncMessage.cpp
@@ -0,0 +1,65 @@
+// Copyright Epic Games, Inc. All Rights Reserved.
+
+#include "AsyncMessage.h"
+
+FAsyncMessage::FAsyncMessage(
+	const FAsyncMessageId& InMessageId,
+	const FAsyncMessageId& InMessageSourceId,
+	const double MessageTimestamp,
+	const uint64 CurrentFrame,
+	const uint32 InThreadQueuedFrom,
+	const uint32 InMessageSequenceCount,
+	FConstStructView InPayloadData)
+	: MessageId(InMessageId)
+	, MessageSourceId(InMessageSourceId)
+	, QueueTime(MessageTimestamp)
+	, QueueFrame(CurrentFrame)
+	, ThreadQueuedFrom(InThreadQueuedFrom)
+	, SequenceId(InMessageSequenceCount)
+	, PayloadCopy(InPayloadData)	// <-- This is making a copy of the payload data!!! $$$
+									// It calls the explicit FInstancedStruct(const FConstStructView InOther) constructor
+									// which is going to copy the UScriptStruct.
+									// TODO: Use some custom linear allocator here on the message system to make the copying cheaper
+{
+	
+}
+
+double FAsyncMessage::GetQueueTimestamp() const
+{
+	return QueueTime;
+}
+
+uint64 FAsyncMessage::GetQueueFrame() const
+{
+	return QueueFrame;
+}
+
+uint32 FAsyncMessage::GetThreadQueuedFromThreadId() const
+{
+	return ThreadQueuedFrom;
+}
+
+uint32 FAsyncMessage::GetSequenceId() const
+{
+	return SequenceId;
+}
+
+FAsyncMessageId FAsyncMessage::GetMessageId() const
+{
+	return MessageId;
+}
+
+void FAsyncMessage::SetMessageId(const FAsyncMessageId& NewMessageId)
+{
+	MessageId = NewMessageId;
+}
+
+FAsyncMessageId FAsyncMessage::GetMessageSourceId() const
+{
+	return MessageSourceId;
+}
+
+FStructView FAsyncMessage::GetPayloadView() const
+{
+	return PayloadCopy.GetScriptStruct();
+}
diff --git a/Engine/Plugins/Experimental/AsyncMessageSystem/Source/AsyncMessageSystem/Private/AsyncMessageBindingOptions.cpp b/Engine/Plugins/Experimental/AsyncMessageSystem/Source/AsyncMessageSystem/Private/AsyncMessageBindingOptions.cpp
new file mode 100644
--- /dev/null
+++ b/Engine/Plugins/Experimental/AsyncMessageSystem/Source/AsyncMessageSystem/Private/AsyncMessageBindingOptions.cpp
@@ -0,0 +1,27 @@
+// Copyright Epic Games, Inc. All Rights Reserved.
+
+#include "AsyncMessageBindingOptions.h"
+
+bool FAsyncMessageBindingOptions::operator==(const FAsyncMessageBindingOptions& Other) const
+{
+	return
+		Type == Other.Type &&
+		Group == Other.Group &&
+		TaskId == Other.TaskId;
+}
+
+bool FAsyncMessageBindingOptions::operator!=(const FAsyncMessageBindingOptions& Other) const
+{
+	return !(*this == Other);
+}
+
+uint32 GetTypeHash(const FAsyncMessageBindingOptions& Opts)
+{
+	uint32 Hash = 0;
+
+	Hash = HashCombine(Hash, GetTypeHash(static_cast<uint8>(Opts.Type)));
+	Hash = HashCombine(Hash, GetTypeHash(static_cast<int32>(Opts.Group)));
+	Hash = HashCombine(Hash, GetTypeHash(Opts.TaskId));
+	
+	return Hash;
+}
diff --git a/Engine/Plugins/Experimental/AsyncMessageSystem/Source/AsyncMessageSystem/Private/AsyncMessageHandle.cpp b/Engine/Plugins/Experimental/AsyncMessageSystem/Source/AsyncMessageSystem/Private/AsyncMessageHandle.cpp
new file mode 100644
--- /dev/null
+++ b/Engine/Plugins/Experimental/AsyncMessageSystem/Source/AsyncMessageSystem/Private/AsyncMessageHandle.cpp
@@ -0,0 +1,60 @@
+// Copyright Epic Games, Inc. All Rights Reserved.
+
+#include "AsyncMessageHandle.h"
+
+const FAsyncMessageHandle FAsyncMessageHandle::Invalid = FAsyncMessageHandle();
+
+bool FAsyncMessageHandle::IsValid() const
+{
+	return InternalHandle != FAsyncMessageHandle::InvalidHandleIndex;
+}
+
+uint32 FAsyncMessageHandle::GetId() const
+{
+	return InternalHandle;
+}
+
+FAsyncMessageId FAsyncMessageHandle::GetBoundMessageId() const
+{
+	return BoundMessage;
+}
+
+FString FAsyncMessageHandle::ToString() const
+{
+	return FString::Printf(TEXT("%u"), InternalHandle);
+}
+
+bool FAsyncMessageHandle::operator==(const FAsyncMessageHandle& Other) const
+{
+	return InternalHandle == Other.InternalHandle; 
+}
+
+bool FAsyncMessageHandle::operator!=(const FAsyncMessageHandle& Other) const
+{
+	return InternalHandle != Other.InternalHandle; 
+}
+
+bool FAsyncMessageHandle::operator>=(const FAsyncMessageHandle& Other) const
+{
+	return InternalHandle >= Other.InternalHandle;
+}
+
+bool FAsyncMessageHandle::operator<(const FAsyncMessageHandle& Other) const
+{
+	return InternalHandle < Other.InternalHandle;
+}
+
+uint32 GetTypeHash(const FAsyncMessageHandle& InMapping)
+{
+	return GetTypeHash(InMapping.InternalHandle);
+}
+
+//////////////////////////////////////////////////////////
+// Internal creation of message handles...
+FAsyncMessageHandle::FAsyncMessageHandle(const uint32 InHandleValue, const FAsyncMessageId InBoundMessage)
+	: InternalHandle(InHandleValue)
+	, BoundMessage(InBoundMessage)
+{
+	checkf(InHandleValue != InvalidHandleIndex, TEXT("'%ud' is an invalid value for FAsyncMessageHandle!"), InHandleValue);
+	checkf(BoundMessage.IsValid(), TEXT("'%s' is an invalid FAsyncMessageId to create a handle for!"), *BoundMessage.ToString());
+}
\ No newline at end of file
diff --git a/Engine/Plugins/Experimental/AsyncMessageSystem/Source/AsyncMessageSystem/Private/AsyncMessageId.cpp b/Engine/Plugins/Experimental/AsyncMessageSystem/Source/AsyncMessageSystem/Private/AsyncMessageId.cpp
new file mode 100644
--- /dev/null
+++ b/Engine/Plugins/Experimental/AsyncMessageSystem/Source/AsyncMessageSystem/Private/AsyncMessageId.cpp
@@ -0,0 +1,74 @@
+// Copyright Epic Games, Inc. All Rights Reserved.
+
+#include "AsyncMessageId.h"
+#include "GameplayTagContainer.h"
+
+// A message with an empty gameplay tag is considered invalid
+const FAsyncMessageId FAsyncMessageId::Invalid = {};
+
+FAsyncMessageId::FAsyncMessageId(const FName MessageName)
+	: InternalMessageTag(FGameplayTag::RequestGameplayTag(MessageName))
+{
+}
+
+FAsyncMessageId::FAsyncMessageId(const FGameplayTag& MessageTag)
+	: InternalMessageTag(MessageTag)
+{
+}
+
+bool FAsyncMessageId::IsValid() const
+{
+	return InternalMessageTag.IsValid();
+}
+
+FName FAsyncMessageId::GetMessageName() const
+{
+	return InternalMessageTag.GetTagName();
+}
+
+FString FAsyncMessageId::ToString() const
+{
+	return InternalMessageTag.ToString();
+}
+
+FAsyncMessageId FAsyncMessageId::GetParentMessageId() const
+{
+	return FAsyncMessageId(InternalMessageTag.RequestDirectParent());
+}
+
+bool FAsyncMessageId::operator==(const FAsyncMessageId& Other) const
+{
+	return InternalMessageTag == Other.InternalMessageTag;
+}
+
+bool FAsyncMessageId::operator!=(const FAsyncMessageId& Other) const
+{
+	return !(*this == Other);
+}
+
+uint32 GetTypeHash(const FAsyncMessageId& InMessageId)
+{
+	// Just use the internal FName for the hash type of 
+	return GetTypeHash(InMessageId.InternalMessageTag);
+}
+
+void FAsyncMessageId::WalkMessageHierarchy(const FAsyncMessageId StartingMessage, TFunction<void(const FAsyncMessageId MessageId)> ForEachMessageFunc)
+{
+	TRACE_CPUPROFILER_EVENT_SCOPE(FAsyncMessageId::WalkMessageHierarchy);
+
+	if (!ensure(StartingMessage.IsValid()))
+	{
+		return;
+	}
+	
+	FAsyncMessageId CurrentMessageId = StartingMessage;	
+
+	while (CurrentMessageId.IsValid())
+	{		
+		// "Visit" this message
+		ForEachMessageFunc(CurrentMessageId);
+
+		// Next we will find the parent message
+		CurrentMessageId = CurrentMessageId.GetParentMessageId();
+	}
+}
\ No newline at end of file
diff --git a/Engine/Plugins/Experimental/AsyncMessageSystem/Source/AsyncMessageSystem/Private/AsyncMessageStore.cpp b/Engine/Plugins/Experimental/AsyncMessageSystem/Source/AsyncMessageSystem/Private/AsyncMessageStore.cpp
new file mode 100644
--- /dev/null
+++ b/Engine/Plugins/Experimental/AsyncMessageSystem/Source/AsyncMessageSystem/Private/AsyncMessageStore.cpp
@@ -0,0 +1,90 @@
+// Copyright Epic Games, Inc. All Rights Reserved.
+
+#include "AsyncMessageStore.h"
+
+uint32 FAsyncMessageStore::EnqueueMessage(FAsyncMessage&& MessageToQueue, TSet<FAsyncMessageBindingOptions>& BindingTypesAddedTo)
+{
+	TRACE_CPUPROFILER_EVENT_SCOPE(IAsyncMessageSystem::ProcessMessageQueueForBinding_Impl);
+	
+	// We want to keep track of which BindingOptions had listeners for this message ID
+	BindingTypesAddedTo.Reset();
+	
+	FScopeLock QueueLock(&MessageQueueCriticalSection);
+	
+	uint32 NumQueuesAddedTo = 0u;
+
+	// Check if there are any listeners to this event or any of its parents. If there are, add it to their binding queue
+	FAsyncMessageId::WalkMessageHierarchy(MessageToQueue.GetMessageId(), [&](const FAsyncMessageId& CurrentMessageId)
+	{
+		// Queue this message to each of the bound options
+		for (TPair<FAsyncMessageBindingOptions, FBindingOptionsMessageQueue>& Pair : MessageBindingQueues)
+		{
+			// Only add this message to a queue if we know that it has a listener for this particular message.
+			if (Pair.Value.MessagesWithTheseBindings.Contains(CurrentMessageId) && !BindingTypesAddedTo.Contains(Pair.Key))
+			{
+				Pair.Value.MessageQueue->Enqueue(MoveTemp(MessageToQueue));
+				++NumQueuesAddedTo;
+			
+				BindingTypesAddedTo.Add(Pair.Key);
+			}
+		}	
+	});
+	
+	return NumQueuesAddedTo;
+}
+
+TOptional<FAsyncMessage> FAsyncMessageStore::GetNextMessageForBindingOption(const FAsyncMessageBindingOptions& Binding)
+{
+	// TODO: We can probably have one big lock for this instead of locking on access internally like this
+	FScopeLock QueueLock(&MessageQueueCriticalSection);	
+
+	if (FBindingOptionsMessageQueue* QueueContainer = MessageBindingQueues.Find(Binding))
+	{
+		FAsyncMessage OutMessage {};
+		if (QueueContainer->MessageQueue->Dequeue(OutMessage))
+		{			
+			return OutMessage;
+		}
+	}
+	
+	return TOptional<FAsyncMessage>{};
+}
+
+void FAsyncMessageStore::AddMessageToBinding(const FAsyncMessageId& MessageId, const FAsyncMessageBindingOptions& Binding)
+{
+	FScopeLock QueueLock(&MessageQueueCriticalSection);
+	
+	FBindingOptionsMessageQueue& QueueData = MessageBindingQueues.FindOrAdd(Binding);
+
+	// Keep track of how many listeners are bound to this message with these binding options
+	uint32& Count = QueueData.MessagesWithTheseBindings.FindOrAdd(MessageId, /* default value */ 0u);
+	++Count;
+}
+
+bool FAsyncMessageStore::RemoveMessageFromBinding(const FAsyncMessageId& MessageId, const FAsyncMessageBindingOptions& Binding)
+{
+	FScopeLock QueueLock(&MessageQueueCriticalSection);
+
+	// If we have no messages that match this binding, then there is nothing to be done
+	FBindingOptionsMessageQueue* QueueData = MessageBindingQueues.Find(Binding);
+	if (!QueueData)
+	{
+		return false;
+	}
+
+	// Decremen the count of how many messages of this ID are on this binding option
+	if (uint32* Count = QueueData->MessagesWithTheseBindings.Find(MessageId))
+	{
+		(*Count)--;
+
+		// If there are no more, then we can remove this message ID from the data store
+		if ((*Count) == 0u)
+		{
+			// Return true if successfully removed
+			return QueueData->MessagesWithTheseBindings.Remove(MessageId) > 0;
+		}
+	}
+
+	// Otherwise there are still other references to this message bound to this binding, so don't remove it
+	return false;
+}
\ No newline at end of file
diff --git a/Engine/Plugins/Experimental/AsyncMessageSystem/Source/AsyncMessageSystem/Private/AsyncMessageStore.h b/Engine/Plugins/Experimental/AsyncMessageSystem/Source/AsyncMessageSystem/Private/AsyncMessageStore.h
new file mode 100644
--- /dev/null
+++ b/Engine/Plugins/Experimental/AsyncMessageSystem/Source/AsyncMessageSystem/Private/AsyncMessageStore.h
@@ -0,0 +1,79 @@
+// Copyright Epic Games, Inc. All Rights Reserved.
+
+#pragma once
+
+#include "AsyncMessageBindingOptions.h"
+
+#include "AsyncMessageId.h"
+#include "AsyncMessage.h"
+#include "Containers/Map.h"
+#include "Containers/Queue.h"
+#include "Templates/SharedPointer.h"
+
+/**
+ * Implementation for the actual storage of async messages for different binding methods.
+ * 
+ * This will store a message queue for each type of supports Binding option for listeners,
+ * allowing listeners to be in control over when they receive messages from the queue (what tick group,
+ * thread ID, etc).
+ */
+class FAsyncMessageStore final
+{
+public:
+	FAsyncMessageStore() = default;
+	~FAsyncMessageStore() = default;
+
+	UE_NONCOPYABLE(FAsyncMessageStore);
+
+public:
+
+	uint32 EnqueueMessage(FAsyncMessage&& MessageToQueue, OUT TSet<FAsyncMessageBindingOptions>& BindingTypesAddedTo);
+
+	[[nodiscard]] TOptional<FAsyncMessage> GetNextMessageForBindingOption(const FAsyncMessageBindingOptions& Binding);
+
+	void AddMessageToBinding(const FAsyncMessageId& MessageId, const FAsyncMessageBindingOptions& Binding);
+
+	/**
+	 * Removes the given message ID from the specified message binding. This is called when an message handler is unbound
+	 * 
+	 */
+	bool RemoveMessageFromBinding(const FAsyncMessageId& MessageId, const FAsyncMessageBindingOptions& Binding);
+
+protected:
+
+	/**
+	* CS for when access to the MessageBindingQueues is required 
+	* to enqueue or dequeue messages accordingly.
+	*/
+	mutable FCriticalSection MessageQueueCriticalSection;
+	
+	// The message queue for a single set of binding options
+	struct FBindingOptionsMessageQueue
+	{
+		FBindingOptionsMessageQueue()
+		{
+			MessageQueue = MakeShared<TQueue<FAsyncMessage>>();
+		}
+
+		~FBindingOptionsMessageQueue()
+		{
+			MessageQueue.Reset();
+		}
+		
+		// Keep track of what messages are bound within this interface so that we can quickly
+		// find which queues require which messages.
+		// The value of this map is how many times it has been bound.
+		TMap<FAsyncMessageId, uint32> MessagesWithTheseBindings;
+
+		// The queue of messages for this binding option		
+		TSharedPtr<TQueue<FAsyncMessage>> MessageQueue;
+	};
+
+	/**
+	 * Tracks the binding options to their associated message queues.
+	 * 
+	 * Each binding option needs its own message queue so that we can enable
+	 * the listeners of the messages to control _when_ they receive an message during the frame. 
+	 */
+	TMap<FAsyncMessageBindingOptions, FBindingOptionsMessageQueue> MessageBindingQueues;
+};
diff --git a/Engine/Plugins/Experimental/AsyncMessageSystem/Source/AsyncMessageSystem/Private/AsyncMessageSystemBase.cpp b/Engine/Plugins/Experimental/AsyncMessageSystem/Source/AsyncMessageSystem/Private/AsyncMessageSystemBase.cpp
new file mode 100644
--- /dev/null
+++ b/Engine/Plugins/Experimental/AsyncMessageSystem/Source/AsyncMessageSystem/Private/AsyncMessageSystemBase.cpp
@@ -0,0 +1,368 @@
+// Copyright Epic Games, Inc. All Rights Reserved.
+
+#include "AsyncMessageSystemBase.h"
+#include "AsyncMessageSystemLogs.h"
+#include "CoreGlobals.h"
+#include "Misc/App.h"
+#include "NativeGameplayTags.h"
+
+FAsyncMessageSystemBase::~FAsyncMessageSystemBase()
+{
+	// Make sure that you call the "Shutdown" function on your message system to allow it to clean up before destruction!
+	ensureMsgf(bIsShuttingDown == true, TEXT("A message system was destructed but did not have its shutdown function called"));
+}
+
+FAsyncMessageHandle FAsyncMessageSystemBase::BindListener(
+	const FAsyncMessageId MessageId,
+	FMessageCallbackFunc&& Callback,
+	const FAsyncMessageBindingOptions& Options /* = {} */)
+{
+	const FAsyncMessageHandle OutHandle = GenerateNextValidMessageHandle(MessageId);
+	
+	const bool bWasSuccessfullyBound = BindListener_Impl(OutHandle, MessageId, MoveTemp(Callback), Options);
+	
+	return bWasSuccessfullyBound ? OutHandle : FAsyncMessageHandle::Invalid;
+}
+
+void FAsyncMessageSystemBase::UnbindListener(const FAsyncMessageHandle& HandleToUnbind)
+{
+	UnbindListener_Impl(HandleToUnbind);
+}
+
+bool FAsyncMessageSystemBase::QueueMessageForBroadcast(const FAsyncMessageId MessageId, FConstStructView PayloadData)
+{
+	TRACE_CPUPROFILER_EVENT_SCOPE(IAsyncMessageSystem::QueueMessageForBroadcast);
+	return QueueMessageForBroadcast_Impl(MessageId, MoveTemp(PayloadData));
+}
+
+void FAsyncMessageSystemBase::ProcessMessageQueueForBinding(const FAsyncMessageBindingOptions& Options)
+{
+	TRACE_CPUPROFILER_EVENT_SCOPE(IAsyncMessageSystem::ProcessMessageQueueForBinding);
+
+	// Handle any pre-processing which may need to happen before sending messages. This could include
+	// binding listeners, removing old handles, and more.
+	PreProcessMessagesQueue(Options);
+
+	// Actually process the messages
+	ProcessMessageQueueForBinding_Impl(Options);
+}
+
+void FAsyncMessageSystemBase::Startup()
+{
+	TRACE_CPUPROFILER_EVENT_SCOPE(IAsyncMessageSystem::Startup);
+	
+	UE_LOG(LogAsyncMessageSystem, Verbose, TEXT("[%hs] Message System Startup..."), __func__);
+	
+	checkf(bIsShuttingDown == false, TEXT("Attempting to restart a message system is not supported."));
+	
+	Startup_Impl();
+}
+
+void FAsyncMessageSystemBase::Shutdown()
+{
+	TRACE_CPUPROFILER_EVENT_SCOPE(IAsyncMessageSystem::Shutdown);
+	
+	UE_LOG(LogAsyncMessageSystem, Verbose, TEXT("[%hs] Message System Shutdown..."), __func__);
+	
+	// Flag this message system as being shut down so that we don't
+	// attempt to queue any more messages or start any new tasks
+	bIsShuttingDown = true;
+
+	Shutdown_Impl();
+}
+
+///////////////////////////////////////////////////////////
+// Virtual Implementations
+
+bool FAsyncMessageSystemBase::BindListener_Impl(const FAsyncMessageHandle& HandleToBindTo, const FAsyncMessageId MessageId, FMessageCallbackFunc&& Callback, const FAsyncMessageBindingOptions& Options)
+{
+	TRACE_CPUPROFILER_EVENT_SCOPE(IAsyncMessageSystem::BindListener_Impl);
+	
+	// Do not allow the binding of new listeners during the shutdown of this message system
+	if (bIsShuttingDown)
+	{
+		UE_LOG(LogAsyncMessageSystem, Error, TEXT("[%hs] Attempting to bind an message listener to message '%s' during shutdown."),
+			__func__, *MessageId.ToString());
+		
+		return false;
+	}
+
+	// Don't allow you to bind to an invalid message id
+	if (!MessageId.IsValid())
+	{
+		UE_LOG(LogAsyncMessageSystem, Error, TEXT("[%hs] Attempting to bind an message listener to invalid message name '%s'"),
+			__func__, *MessageId.ToString());
+		
+		return false;
+	}
+
+	if (!HandleToBindTo.IsValid())
+	{
+		UE_LOG(LogAsyncMessageSystem, Error, TEXT("[%hs] Attempting to bind message '%s' to invalid handle '%s'"),
+			__func__, *MessageId.ToString(), *HandleToBindTo.ToString());
+		
+		return false;
+	}
+
+	// Let the message store know that this message will be listened for by this binding option
+	MessageStore.AddMessageToBinding(MessageId, Options);
+
+	FAsyncMessageIndividualListener ListenerData
+		{
+			.Handle = HandleToBindTo,
+			.Callback = MoveTemp(Callback)
+		};
+
+	// Queue this listener for binding the next time that this system is processed.
+	PendingBoundListenerQueue.Enqueue(
+		FPendingBoundListener<FMessageCallbackFunc> 
+		{
+			.MessageId = MessageId, 
+			.Data = MoveTemp(ListenerData), 
+			.Options = Options 
+		});
+
+	// Successfully queued for binding
+	return true;
+}
+
+void FAsyncMessageSystemBase::UnbindListener_Impl(const FAsyncMessageHandle& HandleToUnbind)
+{	
+	MessageHandlesPendingRemoval.Enqueue(HandleToUnbind);
+}
+
+bool FAsyncMessageSystemBase::QueueMessageForBroadcast_Impl(const FAsyncMessageId MessageId, FConstStructView PayloadData)
+{
+	// Do not allow the queue of new listeners during the shutdown of this message system
+	if (bIsShuttingDown)
+	{
+		UE_LOG(LogAsyncMessageSystem, Error, TEXT("[%hs] Attempting to queue an message '%s' during shutdown. The message will not be queued"),
+			__func__, *MessageId.ToString());
+		
+		return false;
+	}
+
+	// Don't allow you to queue an invalid message id
+	if (!MessageId.IsValid())
+	{
+		UE_LOG(LogAsyncMessageSystem, Error, TEXT("[%hs] Attempting to queue an invalid message id  '%s', the message will not be queued"),
+			__func__, *MessageId.ToString());
+		
+		return false;
+	}
+	
+	// Keep in mind that this function will likely be called from many different threads.
+	// That is why we store our message message queue in a TMpscQueue, because there are multiple
+	// producers (the systems queuing messages) and one consumer (this message system, processing them)
+
+	// The message system will only store WEAK pointers to the bound listeners, and only VIEWS to the payload data.
+	// This means that it is the responsibility of the system calling this "Queue" function to maintain ownership
+	// of the payload data and ensure that it is in scope for any listeners that may want to access it.
+
+	// Note: We have FStructView here, which should give us some reflected data about the type of
+	// payload we are receiving.
+
+	const uint64 QueueTime = FApp::GetCurrentTime();
+	const uint64 QueuedFrame = GFrameCounter;
+	const uint32 ThreadQueuedFrom = FPlatformTLS::GetCurrentThreadId();
+	const uint32 MessageSequence = NextMessageSequence++;
+	
+	uint32 NumQueuesAddedTo = 0u;
+	TSet<FAsyncMessageBindingOptions> BoundTypes;
+	
+	// Push a new async message for this message and all it's parents.
+	FAsyncMessage NewMessage
+		(
+			/* .MessageId= */ MessageId,
+			/* .MessageSourceId= */ MessageId,
+			/* .QueueTime= */ QueueTime,
+			/* .QueueFrame= */ QueuedFrame,
+			/* .ThreadQueuedFrom= */ ThreadQueuedFrom,
+			/* .SequenceId= */ MessageSequence,		
+			/* Payload= */ MoveTemp(PayloadData)		// <--- Note: The payload data is going to be copied for each message in the queue
+		);
+			
+	// Queue this message to the parent message
+	NumQueuesAddedTo += MessageStore.EnqueueMessage(MoveTemp(NewMessage), OUT BoundTypes);
+	
+	UE_LOG(LogAsyncMessageSystem, Verbose, TEXT("[%hs] message '%s' was added to %u binding queues"), 
+		__func__, *MessageId.ToString(), NumQueuesAddedTo);
+
+	// Let subclasses know that an message was just queued and for what binding options. This allows them to 
+	// do things like spin up async tasks (or UE::Tasks) in order to start processing that message queue
+	PostQueueMessage(MessageId, BoundTypes);
+	
+	// Return true if there were any bound listeners to this message. 
+	// False if the message was not added to any message queues
+	return NumQueuesAddedTo > 0;
+}
+
+void FAsyncMessageSystemBase::PreProcessMessagesQueue(const FAsyncMessageBindingOptions& Options)
+{
+	TRACE_CPUPROFILER_EVENT_SCOPE(IAsyncMessageSystem::PreProcessMessagesQueue);
+
+	// Remove any message handles that have been unbound
+	ProcessUnbindHandleRequests();
+
+	// Process any deferred bindings which have been requested
+	ProcessListenersPendingBinding();	
+}
+
+void FAsyncMessageSystemBase::ProcessMessageQueueForBinding_Impl(const FAsyncMessageBindingOptions& Options)
+{
+	TRACE_CPUPROFILER_EVENT_SCOPE(IAsyncMessageSystem::ProcessMessageQueueForBinding_Impl);
+	
+	// Reset the message sequence ID which we are processing this frame
+	NextMessageSequence = 0;
+
+	// Lock the message listener maps so that listeners don't get added/removed
+	// in the middle of processing, which would cause a data race
+	FScopeLock ListenerLock(&MessageListenerMapCS);
+	
+	uint32 ProcessedMessages = 0;
+	uint32 MessagesCalled = 0;
+	
+	// Pop off of the message queue and process each message
+	TOptional<FAsyncMessage> OptionalNextMessage = MessageStore.GetNextMessageForBindingOption(Options);
+	while (OptionalNextMessage.IsSet())
+	{
+		FAsyncMessage& Message = OptionalNextMessage.GetValue();
+
+		// TODO: If we define a "lifetime" for messages, we could grab the current frame/time
+		// here and check against it to see if enough time has passed and we should just "throw
+		// out" the message
+		//
+		//const double CurrentTime = FApp::GetCurrentTime();
+		//const uint64 CurrentFrame = GFrameCounter;
+		
+		// For each message in the heirarchy of this message...
+		FAsyncMessageId::WalkMessageHierarchy(Message.GetMessageId(), [&](const FAsyncMessageId CurrentMessageId)
+		{
+			// Set this message ID to the current message id
+			Message.SetMessageId(CurrentMessageId);
+			
+			// If we know of any bound listeners to this message, notify them.
+			if (FAsyncMessageBoundData* BoundData = BoundMessageListenerMap.Find(Message.GetMessageId()))
+			{
+				// Look for any listeners associated with this binding option
+				if (TArray<FAsyncMessageIndividualListener<FMessageCallbackFunc>>* Listeners = BoundData->ListenerMap.Find(Options))
+				{
+					for (const FAsyncMessageIndividualListener<FMessageCallbackFunc>& Listener : *Listeners)
+					{
+						// As long as the callback for this listener is valid, broadcast it
+						if (Listener.Callback)
+						{
+							// Broadcast to the listener!					
+							Listener.Callback(Message);
+						
+							MessagesCalled++;	
+						}
+						// Otherwise, this callback is no long valid (the lambda may have gone out of scope)
+						// and we should unbind it so that we don't bother trying to call it again
+						else
+						{
+							UnbindListener(Listener.Handle);
+						}
+					}
+				}
+			}
+		});
+		
+		OptionalNextMessage = MessageStore.GetNextMessageForBindingOption(Options);
+	}
+
+	UE_CLOG(ProcessedMessages > 0, LogAsyncMessageSystem, VeryVerbose, TEXT("[%hs] Processed %u messages. Called %u message listeners."), 
+		__func__, ProcessedMessages, MessagesCalled);
+}
+
+void FAsyncMessageSystemBase::ProcessUnbindHandleRequests()
+{
+	TRACE_CPUPROFILER_EVENT_SCOPE(IAsyncMessageSystem::ProcessUnbindHandleRequests);
+
+	if (MessageHandlesPendingRemoval.IsEmpty())
+	{
+		return;
+	}
+
+	FScopeLock MessageListenerMapLock(&MessageListenerMapCS);
+
+	TOptional<FAsyncMessageHandle> HandleToRemoveOpt = MessageHandlesPendingRemoval.Dequeue();
+	while (HandleToRemoveOpt.IsSet())
+	{
+		FAsyncMessageHandle& HandleToRemove = HandleToRemoveOpt.GetValue();
+
+		if (FAsyncMessageBoundData* FoundBindingData = BoundMessageListenerMap.Find(HandleToRemove.GetBoundMessageId()))
+		{
+			// Check the listener map for this message...
+			for (auto& BindingPair : FoundBindingData->ListenerMap)
+			{
+				TArray<FAsyncMessageIndividualListener<FMessageCallbackFunc>>& Listeners = BindingPair.Value;
+
+				// Remove any listeners whose handles have been marked for unbinding
+				for (int32 i = Listeners.Num() - 1; i >= 0; --i)
+				{
+					if (Listeners[i].Handle == HandleToRemove)
+					{
+						MessageStore.RemoveMessageFromBinding(Listeners[i].Handle.GetBoundMessageId(), BindingPair.Key);
+						Listeners.RemoveAt(i);
+					}
+				}
+			}
+		}
+
+		// Get the next handle to remove
+		HandleToRemoveOpt = MessageHandlesPendingRemoval.Dequeue();
+	}
+}
+
+void FAsyncMessageSystemBase::ProcessListenersPendingBinding()
+{
+	TRACE_CPUPROFILER_EVENT_SCOPE(IAsyncMessageSystem::ProcessListenersPendingBinding);
+
+	if (PendingBoundListenerQueue.IsEmpty())
+	{
+		return;
+	}
+
+	// We will lock the listener map and binding of messages to allow for multiple threads
+	// to add listeners at the same time safely
+	FScopeLock ListenerLock(&MessageListenerMapCS);
+
+	// Iterate through the pending listeners and add them to our BoundMessageListenerMap Map.
+	TOptional<FPendingBoundListener<FMessageCallbackFunc>> ListenerToBind = PendingBoundListenerQueue.Dequeue();
+	while (ListenerToBind.IsSet())
+	{
+		FAsyncMessageBoundData& MessageData = BoundMessageListenerMap.FindOrAdd(ListenerToBind->MessageId);
+
+		// Add a new listener which we can broadcast the callback to later
+		MessageData.ListenerMap.FindOrAdd(ListenerToBind->Options).Add(MoveTemp(ListenerToBind->Data));
+
+		ListenerToBind = PendingBoundListenerQueue.Dequeue();
+	}
+}
+
+FAsyncMessageHandle FAsyncMessageSystemBase::GenerateNextValidMessageHandle(const FAsyncMessageId ForMessageId)
+{
+	// Increment to the next handle, checking if its value is zero.
+	// If the value is zero, that means that we have hit uint32 max and overflowed back around
+	// and we should increment the handle again to get a valid index
+	uint32 HandleValue = ++NextMessageHandleId;
+	if (HandleValue == 0)
+	{
+		HandleValue = ++NextMessageHandleId;
+		
+		// It is likely non-fatal if the number of handles has grown this pointer because the older handles will
+		// be safe to be re-used, but log a warning here in case
+		UE_LOG(LogAsyncMessageSystem, Warning, TEXT("The Async Message Handle index has been wrapped!"));
+	}
+	
+	return FAsyncMessageHandle(HandleValue, ForMessageId);
+}
+
+// Dev test function only to test creation of valid handles, just to make unit testing easier.
+#if WITH_DEV_AUTOMATION_TESTS
+FAsyncMessageHandle FAsyncMessageSystemBase::GenerateHandleAtIndex(const uint32 Index, const FAsyncMessageId& ForId)
+{
+	return FAsyncMessageHandle(Index, ForId);
+}
+#endif
\ No newline at end of file
diff --git a/Engine/Plugins/Experimental/AsyncMessageSystem/Source/AsyncMessageSystem/Private/AsyncMessageSystemLogs.cpp b/Engine/Plugins/Experimental/AsyncMessageSystem/Source/AsyncMessageSystem/Private/AsyncMessageSystemLogs.cpp
new file mode 100644
--- /dev/null
+++ b/Engine/Plugins/Experimental/AsyncMessageSystem/Source/AsyncMessageSystem/Private/AsyncMessageSystemLogs.cpp
@@ -0,0 +1,5 @@
+// Copyright Epic Games, Inc. All Rights Reserved.
+
+#include "AsyncMessageSystemLogs.h"
+
+DEFINE_LOG_CATEGORY(LogAsyncMessageSystem);
\ No newline at end of file
diff --git a/Engine/Plugins/Experimental/AsyncMessageSystem/Source/AsyncMessageSystem/Private/AsyncMessageSystemModule.cpp b/Engine/Plugins/Experimental/AsyncMessageSystem/Source/AsyncMessageSystem/Private/AsyncMessageSystemModule.cpp
new file mode 100644
--- /dev/null
+++ b/Engine/Plugins/Experimental/AsyncMessageSystem/Source/AsyncMessageSystem/Private/AsyncMessageSystemModule.cpp
@@ -0,0 +1,14 @@
+// Copyright Epic Games, Inc. All Rights Reserved.
+
+#include "AsyncMessageSystemModule.h"
+#include "AsyncGameplayMessageSystem.h"
+#include "AsyncMessageSystemLogs.h"
+#include "AsyncMessageSystemBase.h"
+#include "Modules/ModuleManager.h"
+
+bool FAsyncMessageSystemModule::IsAvailable()
+{
+	return FModuleManager::Get().IsModuleLoaded("AsyncMessageSystem");
+}
+
+IMPLEMENT_MODULE(FAsyncMessageSystemModule, AsyncMessageSystem)
diff --git a/Engine/Plugins/Experimental/AsyncMessageSystem/Source/AsyncMessageSystem/Private/AsyncMessageSystemModule.h b/Engine/Plugins/Experimental/AsyncMessageSystem/Source/AsyncMessageSystem/Private/AsyncMessageSystemModule.h
new file mode 100644
--- /dev/null
+++ b/Engine/Plugins/Experimental/AsyncMessageSystem/Source/AsyncMessageSystem/Private/AsyncMessageSystemModule.h
@@ -0,0 +1,20 @@
+// Copyright Epic Games, Inc. All Rights Reserved.
+
+#pragma once
+
+#include "Modules/ModuleInterface.h"
+#include "Templates/SharedPointer.h"
+
+class FAsyncMessageSystemBase;
+
+class FAsyncMessageSystemModule : public IModuleInterface
+{
+public:
+
+	/**
+	 * Checks to see if this module is loaded and ready.  It is only valid to call Get() if IsAvailable() returns true.
+	 *
+	 * @return True if the module is loaded and ready to use
+	 */
+	static bool IsAvailable();
+};
\ No newline at end of file
diff --git a/Engine/Plugins/Experimental/AsyncMessageSystem/Source/AsyncMessageSystem/Private/AsyncMessageWorldSubsystem.cpp b/Engine/Plugins/Experimental/AsyncMessageSystem/Source/AsyncMessageSystem/Private/AsyncMessageWorldSubsystem.cpp
new file mode 100644
--- /dev/null
+++ b/Engine/Plugins/Experimental/AsyncMessageSystem/Source/AsyncMessageSystem/Private/AsyncMessageWorldSubsystem.cpp
@@ -0,0 +1,40 @@
+// Copyright Epic Games, Inc. All Rights Reserved.
+
+#include "AsyncMessageWorldSubsystem.h"
+#include "AsyncMessageSystemLogs.h"
+#include "AsyncGameplayMessageSystem.h"
+
+void UAsyncMessageWorldSubsystem::Initialize(FSubsystemCollectionBase& Collection)
+{
+	Super::Initialize(Collection);
+
+	InitMessageSystem();
+}
+
+void UAsyncMessageWorldSubsystem::Deinitialize()
+{
+	Super::Deinitialize();
+
+	ShutdownMessageSystem();
+}
+
+void UAsyncMessageWorldSubsystem::InitMessageSystem()
+{
+	check(!MessageSystem.IsValid());
+	
+	UE_LOG(LogAsyncMessageSystem, Log, TEXT("[%hs] Init world message system for world %s..."), __func__, *GetNameSafe(GetWorld()));
+
+	MessageSystem = FAsyncMessageSystemBase::CreateMessageSystem<FAsyncGameplayMessageSystem>(GetWorld());
+}
+
+void UAsyncMessageWorldSubsystem::ShutdownMessageSystem()
+{
+	if (ensure(MessageSystem.IsValid()))
+	{
+		UE_LOG(LogAsyncMessageSystem, Log, TEXT("[%hs] Shutdown world message system for world %s..."), __func__, *GetNameSafe(GetWorld()));
+		
+		MessageSystem->Shutdown();
+		
+		MessageSystem.Reset();
+	}
+}
diff --git a/Engine/Plugins/Experimental/AsyncMessageSystem/Source/AsyncMessageSystem/Private/Tests/AsyncMessageSystemTests.cpp b/Engine/Plugins/Experimental/AsyncMessageSystem/Source/AsyncMessageSystem/Private/Tests/AsyncMessageSystemTests.cpp
new file mode 100644
--- /dev/null
+++ b/Engine/Plugins/Experimental/AsyncMessageSystem/Source/AsyncMessageSystem/Private/Tests/AsyncMessageSystemTests.cpp
@@ -0,0 +1,687 @@
+// Copyright Epic Games, Inc. All Rights Reserved.
+
+#include "AsyncMessageSystemTests.h"
+#include "Misc/AutomationTest.h"
+#include "AsyncMessageHandle.h"
+#include "Tasks/Task.h"
+#include "NativeGameplayTags.h"
+
+UE_DEFINE_GAMEPLAY_TAG_COMMENT(InternalTestTag_Invalid, "AsyncMessages.Internal.test.invalid", "A test gameplay tag utilized in the async message system unit tests")
+UE_DEFINE_GAMEPLAY_TAG_COMMENT(InternalTestTag_A, "AsyncMessages.Internal.test.a", "A test gameplay tag utilized in the async message system unit tests")
+UE_DEFINE_GAMEPLAY_TAG_COMMENT(InternalTestTag_B, "AsyncMessages.Internal.test.b", "A test gameplay tag utilized in the async message system unit tests")
+
+UE_DEFINE_GAMEPLAY_TAG_COMMENT(InternalTestTag_Child, "AsyncMessages.Internal.test.child", "A test gameplay tag utilized in the async message system unit tests")
+
+#if WITH_DEV_AUTOMATION_TESTS
+
+namespace UE::AsyncMessageSystem
+{
+	constexpr EAutomationTestFlags QuickTestFlags = EAutomationTestFlags_ApplicationContextMask | EAutomationTestFlags::SmokeFilter;
+	constexpr EAutomationTestFlags StressTestFlags = EAutomationTestFlags_ApplicationContextMask | EAutomationTestFlags::StressFilter;
+
+	IMPLEMENT_SIMPLE_AUTOMATION_TEST(TMessageHandleDefaults, "AsyncMessagePassing.MessageHandles", QuickTestFlags)
+	bool TMessageHandleDefaults::RunTest(const FString& InParameters)
+	{
+		// Test defaults
+		{
+			FAsyncMessageHandle DefaultHandle = {};
+			TestFalse(TEXT("Default async message handle is invalid"), DefaultHandle.IsValid());
+
+			const FString DefaultHandleToString = DefaultHandle.ToString();
+			const FString ExpectedToString = TEXT("0");
+			TestTrue(TEXT("Default handle has correct ToString"), DefaultHandleToString == ExpectedToString);
+		}
+
+		// Test Invalidate
+		{
+			FAsyncMessageHandle TestHandle = FAsyncMessageSystemBase::GenerateHandleAtIndex(123, {InternalTestTag_Invalid});
+			TestTrue(TEXT("Test handle starts as a valid handle"), TestHandle.IsValid());
+		}
+
+		// Test GetId()
+		{
+			constexpr uint32 TestHandleIdx = 456;
+			FAsyncMessageHandle Handle = FAsyncMessageSystemBase::GenerateHandleAtIndex(TestHandleIdx, {InternalTestTag_Invalid});
+			TestTrue(TEXT("Handle Is Valid"), Handle.IsValid());
+			
+			TestTrue(TEXT("Handle Is the same as the given index"), TestHandleIdx == Handle.GetId());
+		}
+
+		// Test GetId()
+		{
+			constexpr uint32 TestHandleIdx = 456;
+			FAsyncMessageHandle Handle = FAsyncMessageSystemBase::GenerateHandleAtIndex(TestHandleIdx, {InternalTestTag_Invalid});
+			
+			const FString DefaultHandleToString = Handle.ToString();
+			const FString ExpectedToString = TEXT("456");
+			TestTrue(TEXT("Handle handle has correct ToString"), DefaultHandleToString == ExpectedToString);
+		}
+
+		// Test handle == operator
+		{
+			FAsyncMessageHandle A = FAsyncMessageSystemBase::GenerateHandleAtIndex(789, {InternalTestTag_Invalid});
+			FAsyncMessageHandle B = FAsyncMessageSystemBase::GenerateHandleAtIndex(789, {InternalTestTag_Invalid});
+			TestTrue(TEXT("Async Message Handle == operator works"), A == B);
+		}
+
+		// Test handle != operator
+		{
+			FAsyncMessageHandle A = FAsyncMessageSystemBase::GenerateHandleAtIndex(123, {InternalTestTag_Invalid});
+			FAsyncMessageHandle B = FAsyncMessageSystemBase::GenerateHandleAtIndex(789, {InternalTestTag_Invalid});
+			TestTrue(TEXT("Async Message Handle != operator works"), A != B);
+		}
+
+		return true;
+	}
+
+	IMPLEMENT_SIMPLE_AUTOMATION_TEST(TMessageParentHandles, "AsyncMessagePassing.MessageParentHandle", QuickTestFlags)
+	bool TMessageParentHandles::RunTest(const FString& InParameters)
+	{
+		static const FName ExpectedChildMessageName					= "AsyncMessages.Internal.test.child";
+		static const FName ExpectedParentMessageName				= "AsyncMessages.Internal.test";
+		static const FName ExpectedGrandparentMessageName			= "AsyncMessages.Internal";
+		static const FName ExpectedGreatGrandparentMessageName		= "AsyncMessages";
+		static const FName ExpectedGreatGreatGrandparentMessageName	= NAME_None;
+		
+		// Start with a child
+		static const FAsyncMessageId GrandChildMessage { InternalTestTag_Child };
+		TestTrue(TEXT("Child name is correct"), ExpectedChildMessageName == GrandChildMessage.GetMessageName());
+		
+		// Ensure the parent is correct
+		const FAsyncMessageId TestParentId = GrandChildMessage.GetParentMessageId();
+		TestTrue(TEXT("Test Parent ID is correct"), ExpectedParentMessageName == TestParentId.GetMessageName());
+
+		// Ensure the grandparent is correct
+		const FAsyncMessageId TestGrandparentId = TestParentId.GetParentMessageId();
+		TestTrue(TEXT("Test Grand Parent ID is correct"), ExpectedGrandparentMessageName == TestGrandparentId.GetMessageName());
+
+		// Great grand parent, which should be the root tag
+		const FAsyncMessageId TestGreatGrandparentId = TestGrandparentId.GetParentMessageId();
+		TestTrue(TEXT("Test Great Grand Parent ID is correct"), ExpectedGreatGrandparentMessageName == TestGreatGrandparentId.GetMessageName());
+
+		// Great great grandparent should be empty
+		const FAsyncMessageId TestGreatGreatGrandparentId = TestGreatGrandparentId.GetParentMessageId();
+		TestTrue(TEXT("Test Great Grand Parent ID is correct"), ExpectedGreatGreatGrandparentMessageName == TestGreatGreatGrandparentId.GetMessageName());
+		
+		return true;
+	}
+
+	const FAsyncMessageId TestMessageId_A = {InternalTestTag_A};
+	
+	struct FTestListener
+	{
+	public:
+		void Callback_Message_A(const FAsyncMessage& Message)
+		{
+			FConstStructView DataView = Message.GetPayloadView();
+			const FTest_Payload_A* Data = Message.GetPayloadData<const FTest_Payload_A>();
+			if (Data)
+			{
+				CallbackMutation += Data->IncrementAmount;
+			}
+		}
+
+		int32 CallbackMutation = 0;
+	};
+
+	// Test that the shutdown flag is correctly being set when we call Shutdown
+	IMPLEMENT_SIMPLE_AUTOMATION_TEST(TMessageSystemShutdownTest, "AsyncMessagePassing.StartupShutdown", QuickTestFlags)
+	bool TMessageSystemShutdownTest::RunTest(const FString& InParameters)
+	{
+		// Startup a new message system...
+		TSharedPtr<FTestMessageSystem, ESPMode::ThreadSafe> TestSystem = FAsyncMessageSystemBase::CreateMessageSystem<FTestMessageSystem>();
+		TestTrue(TEXT("MessageSystem startup has been called"), TestSystem->HasStarted());
+
+		TestSystem->RunOnce();
+		
+		TestSystem->Shutdown();
+		TestTrue(TEXT("MessageSystem shutdown has been called"), TestSystem->HasShutdown());
+
+		TestSystem.Reset();
+		return true;
+	}
+	
+	// Test that we can bind an message onto a non-UObject shared pointer and that it receives the message after we pump the message queue
+	IMPLEMENT_SIMPLE_AUTOMATION_TEST(TMessageSystemBindingSharedPtrTest, "AsyncMessagePassing.Binding.SharedPtr", QuickTestFlags)
+	bool TMessageSystemBindingSharedPtrTest::RunTest(const FString& InParameters)
+	{
+		// Startup a new message system...
+		TSharedPtr<FTestMessageSystem, ESPMode::ThreadSafe> TestSystem = FAsyncMessageSystemBase::CreateMessageSystem<FTestMessageSystem>();
+		TestTrue(TEXT("MessageSystem startup has been called"), TestSystem->HasStarted());
+
+		// Set up a test listener object with some dummy data
+		TSharedPtr<FTestListener> TestListener = MakeShared<FTestListener>();
+		TestTrue(TEXT("Callback Object mutated successfully"), TestListener->CallbackMutation == 0);
+
+		const int32 TestStartValue = TestListener->CallbackMutation;
+		
+		// Bind an message callback to the test listener
+		FAsyncMessageHandle ListenerHandle = TestSystem->BindListener(TestMessageId_A, TestListener.ToWeakPtr(), &FTestListener::Callback_Message_A);
+		
+		// This handle will be in the pending queue
+		TestTrue(TEXT("A valid listener handle was provided to binding"), ListenerHandle.IsValid());
+
+		// Make some instance struct payload
+		const int32 TestAmountToAdd = 1563875499;
+		
+		FTest_Payload_A PayloadData = {};
+		PayloadData.IncrementAmount = TestAmountToAdd;
+		FInstancedStruct PayloadDataInstance = FInstancedStruct::Make<FTest_Payload_A>(PayloadData);
+		
+		// Queue a new message for broadcast!
+		TestSystem->QueueMessageForBroadcast(TestMessageId_A, PayloadDataInstance);
+
+		// Next, simulate the message system being ticked once
+		// This should process our delegate bindings and execute anything in the message queue
+		TestSystem->RunOnce();
+
+		// Ensure that the callback incremented the data with the correct number
+		TestTrue(TEXT("Callback Object mutated successfully"), TestListener->CallbackMutation == (TestStartValue + TestAmountToAdd));
+
+		// Remove the bound listener and test that is was successfully
+		TestSystem->UnbindListener(ListenerHandle);
+
+		TestSystem->RunOnce();
+		TestFalse(TEXT("Unregister the listener handle"), TestSystem->IsHandleBound(ListenerHandle));
+
+		// All done
+		TestSystem->Shutdown();
+		TestSystem.Reset();
+		return true;
+	}
+	
+	// Test to make sure that you can bind an message listener whilst the message system is currently processing
+	IMPLEMENT_SIMPLE_AUTOMATION_TEST(TMessageSystemNestedBinding, "AsyncMessagePassing.Binding.NestedMessage", QuickTestFlags)
+	bool TMessageSystemNestedBinding::RunTest(const FString& InParameters)
+	{
+		static const FAsyncMessageId NestedMessageId = {InternalTestTag_B};
+
+		struct FTestNestedObject
+		{
+			void TestNestedBinding(const FAsyncMessage& Message)
+			{
+				// We should be able to add a test binding in response to 
+				if (const FNested_Payload* Data = Message.GetPayloadData<const FNested_Payload>())
+				{
+					if (TSharedPtr<FTestMessageSystem> Sys = Data->MessageSystem.Pin())
+					{
+						// Bind a nested listener here
+						const FAsyncMessageHandle Handle = Sys->BindListener(TestMessageId_A, [](const FAsyncMessage& Message)
+						{
+							// we dont need to do anything here, just ensure that the actual BindListener function works
+						});	
+					}
+				}
+			}
+		};
+		
+		// Startup a new message system...
+		TSharedPtr<FTestMessageSystem, ESPMode::ThreadSafe> TestSystem = FAsyncMessageSystemBase::CreateMessageSystem<FTestMessageSystem>();
+		TestTrue(TEXT("MessageSystem startup has been called"), TestSystem->HasStarted());
+		
+		TSharedPtr<FTestNestedObject> TestObject = MakeShared<FTestNestedObject>();
+
+		const FAsyncMessageHandle Handle = TestSystem->BindListener(NestedMessageId, TestObject.ToWeakPtr(), &FTestNestedObject::TestNestedBinding);
+		
+		TestTrue(TEXT("System has a no listeners yet, they are in the pending queue"), TestSystem->GetNumberOfListeners() == 0);
+
+		FNested_Payload Payload = { .MessageSystem = TestSystem };
+		FInstancedStruct PayloadInstance = FInstancedStruct::Make<FNested_Payload>(Payload);
+
+		TestSystem->QueueMessageForBroadcast(NestedMessageId, PayloadInstance);
+
+		// Running once should put the nested binding in the "pending" bindings queue
+		TestSystem->RunOnce();
+
+		TestTrue(TEXT("System has a single listener"), TestSystem->GetNumberOfListeners() == 1);
+
+		// Running again should process that pending bindings queue and add the new listener
+		TestSystem->RunOnce();
+		
+		TestTrue(TEXT("System has two listeneres now, one is nested"), TestSystem->GetNumberOfListeners() == 2);
+
+		TestSystem->Shutdown();
+		TestSystem.Reset();
+		return true;
+	}
+
+	// Test that we can bind an message onto a UObject and that it receives the message after we pump the message queue
+	IMPLEMENT_SIMPLE_AUTOMATION_TEST(TMessageSystemBindingUObjectTest, "AsyncMessagePassing.Binding.UObject", QuickTestFlags)
+	bool TMessageSystemBindingUObjectTest::RunTest(const FString& InParameters)
+	{
+		TSharedPtr<FTestMessageSystem, ESPMode::ThreadSafe> TestSystem = FAsyncMessageSystemBase::CreateMessageSystem<FTestMessageSystem>();
+		TestTrue(TEXT("MessageSystem startup has been called"), TestSystem->HasStarted());
+
+		TObjectPtr<UTestAsyncObject> TestObject = NewObject<UTestAsyncObject>();
+		TestTrue(TEXT("Created test object"), IsValid(TestObject));
+		const int32 TestStartValue = TestObject->TestValue;
+
+		// Bind a listener to the test object
+		FAsyncMessageHandle ListenerHandle = TestSystem->BindListener(TestMessageId_A, TWeakObjectPtr<UTestAsyncObject>(TestObject), &UTestAsyncObject::CallbackFunction);
+		TestTrue(TEXT("A valid listener handle was provided to binding"), ListenerHandle.IsValid());
+		
+		// Create a test payload instanced struct
+		const int32 TestAmountToAdd = 7;
+		
+		FTest_Payload_A PayloadData = {};
+		PayloadData.IncrementAmount = TestAmountToAdd;
+		FInstancedStruct PayloadDataInstance = FInstancedStruct::Make<FTest_Payload_A>(PayloadData);
+
+		// Queue a new message for broadcast!
+		TestSystem->QueueMessageForBroadcast(TestMessageId_A, PayloadDataInstance);
+
+		// Next, simulate the message system being ticked once
+		// This should process our delegate bindings and execute anything in the message queue
+		TestSystem->RunOnce();
+
+		TestTrue(TEXT("Callback UObject mutated successfully"), TestObject->TestValue == (TestStartValue + TestAmountToAdd));
+
+		TestSystem->Shutdown();
+		TestSystem.Reset();
+		
+		return true;
+	}	
+
+	// Test binding a simple C++ lambda 
+	IMPLEMENT_SIMPLE_AUTOMATION_TEST(TMessageSystemBindingLambdaTest, "AsyncMessagePassing.Binding.Lambda", QuickTestFlags)
+	bool TMessageSystemBindingLambdaTest::RunTest(const FString& InParameters)
+	{
+		TSharedPtr<FTestMessageSystem, ESPMode::ThreadSafe> TestSystem = FAsyncMessageSystemBase::CreateMessageSystem<FTestMessageSystem>();
+		TestTrue(TEXT("MessageSystem startup has been called"), TestSystem->HasStarted());
+
+		// Create a test payload instanced struct
+		const int32 TestAmountToAdd = 7;
+		
+		FTest_Payload_A PayloadData = {};
+		PayloadData.IncrementAmount = TestAmountToAdd;
+		FInstancedStruct PayloadDataInstance = FInstancedStruct::Make<FTest_Payload_A>(PayloadData);
+
+		// This Test Data will be what we use to verify that the lambda function has run and successfully
+		// read/write some data
+		int32 TestDataToMutate = 5;
+		const int32 TestStartValue = TestDataToMutate;
+		
+		// Bind a listener to the test lambda
+		const FAsyncMessageHandle ListenerHandle = TestSystem->BindListener(TestMessageId_A, [&TestDataToMutate](const FAsyncMessage& Message)
+		{
+			if (const FTest_Payload_A* Data = Message.GetPayloadData<const FTest_Payload_A>())
+			{
+				TestDataToMutate += Data->IncrementAmount;
+			}
+		});
+
+		// Queue a new message for broadcast!
+		TestSystem->QueueMessageForBroadcast(TestMessageId_A, PayloadDataInstance);
+
+		// Next, simulate the message system being ticked once
+		// This should process our delegate bindings and execute anything in the message queue
+		TestSystem->RunOnce();
+
+		TestTrue(TEXT("Callback lambda mutated successfully"), TestDataToMutate == (TestStartValue + TestAmountToAdd));
+		TestSystem->Shutdown();
+		TestSystem.Reset();
+		return true;
+	}
+	
+	// Test that the handle of an message is successfully unbound
+	IMPLEMENT_SIMPLE_AUTOMATION_TEST(TMessageSystemUnBindingSingle, "AsyncMessagePassing.UnBinding.Single", QuickTestFlags)
+	bool TMessageSystemUnBindingSingle::RunTest(const FString& InParameters)
+	{
+		// Startup a new message system...
+		TSharedPtr<FTestMessageSystem, ESPMode::ThreadSafe> TestSystem = FAsyncMessageSystemBase::CreateMessageSystem<FTestMessageSystem>();
+		TestTrue(TEXT("MessageSystem startup has been called"), TestSystem->HasStarted());
+		
+		const FAsyncMessageHandle ListenerHandle = TestSystem->BindListener(TestMessageId_A, [this](const FAsyncMessage& Message)
+		{
+			const uint64 Frame = Message.GetQueueFrame();
+			TestTrue(TEXT("Testing inside the lambda"), Frame > 0);
+		});
+
+		// Similate running once, and ensure that the handle is still bound correctly
+		TestSystem->RunOnce();
+		TestTrue(TEXT("A valid listener handle was provided to binding"), ListenerHandle.IsValid());
+		TestTrue(TEXT("Listener Handle Is bound"), TestSystem->IsHandleBound(ListenerHandle));
+
+		// Remove the bound listener and test that is was successfully
+		TestSystem->UnbindListener(ListenerHandle);
+		TestSystem->RunOnce();
+		TestFalse(TEXT("Unregister the listener handle"), TestSystem->IsHandleBound(ListenerHandle));
+
+		TestSystem->RunOnce();
+		
+		TestFalse(TEXT("Listener handle is no longer valid  after running the system"), TestSystem->IsHandleBound(ListenerHandle));
+
+		// All done
+		TestSystem->Shutdown();
+		TestSystem.Reset();
+		return true;
+	}
+
+	// Test that the handle of an message is successfully unbound
+	IMPLEMENT_SIMPLE_AUTOMATION_TEST(TMessageSystemUnBindingMultiple, "AsyncMessagePassing.UnBinding.Multiple", QuickTestFlags)
+	bool TMessageSystemUnBindingMultiple::RunTest(const FString& InParameters)
+	{
+		// Startup a new message system...
+		TSharedPtr<FTestMessageSystem, ESPMode::ThreadSafe> TestSystem = FAsyncMessageSystemBase::CreateMessageSystem<FTestMessageSystem>();
+		TestTrue(TEXT("MessageSystem startup has been called"), TestSystem->HasStarted());
+
+		TArray<FAsyncMessageHandle> Handles;
+
+		constexpr int32 NumHandlesToBind = 50;
+		int32 TestNum = 0;
+		for (int32 i = 0; i < NumHandlesToBind; i++)
+		{
+			const FAsyncMessageHandle ListenerHandle = TestSystem->BindListener(TestMessageId_A, [&TestNum](const FAsyncMessage& Message)
+			{
+				TestNum++;
+			});
+			Handles.Emplace(ListenerHandle);
+
+			TestTrue(TEXT("Listener Handle Is bound"), ListenerHandle.IsValid());
+		}
+
+		TestSystem->QueueMessageForBroadcast(TestMessageId_A);
+		
+		// Simulate running once, and ensure that the handle is still bound correctly
+		TestSystem->RunOnce();
+		
+		TestTrue(TEXT("Modifyed value correctly"), TestNum == NumHandlesToBind);
+		
+		for (FAsyncMessageHandle& ListenerHandle : Handles)
+		{
+			TestSystem->UnbindListener(ListenerHandle);
+		}
+		
+		TestSystem->RunOnce();
+
+		for (FAsyncMessageHandle& ListenerHandle : Handles)
+		{
+			TestFalse(TEXT("Unregister the listener handle"), TestSystem->IsHandleBound(ListenerHandle));
+		}
+
+		// All done
+		TestSystem->Shutdown();
+		TestSystem.Reset();
+		return true;
+	}
+
+	////////////////////////////////////////////////////
+	// Multi-threaded tests. These will test some core functionality and thread safety by spawning a
+	// bunch of UE::Tasks on different worker threads, and then modifying/accessing the message system
+	// through them.
+	
+	// Test to queue messages from multiple threads on a single message system
+	IMPLEMENT_SIMPLE_AUTOMATION_TEST(TMessageSystemMultiThreadedQueueMessagesTest, "AsyncMessagePassing.MultiThreaded.QueueMessages", StressTestFlags)
+	bool TMessageSystemMultiThreadedQueueMessagesTest::RunTest(const FString& InParameters)
+	{
+		TSharedPtr<FTestMessageSystem, ESPMode::ThreadSafe> TestSystem = FAsyncMessageSystemBase::CreateMessageSystem<FTestMessageSystem>();
+		TestTrue(TEXT("MessageSystem startup has been called"), TestSystem->HasStarted());
+
+		FAsyncMessageHandle ListenerHandle = TestSystem->BindListener(TestMessageId_A, [](const FAsyncMessage& Message)
+		{
+			FConstStructView DataView = Message.GetPayloadView();
+			const FTest_Payload_A* Data = Message.GetPayloadData<const FTest_Payload_A>();
+			if (Data)
+			{
+				UE_LOG(LogTemp, Log, TEXT("DataVal : %d"), Data->IncrementAmount);
+			}
+			else
+			{
+				UE_LOG(LogTemp, Error, TEXT("Failed to find the payload data, it has expired"));
+			}
+		});
+		
+		struct FQueueMessageFunctor
+		{
+			TWeakPtr<FTestMessageSystem> MessageSys;
+
+			int32 AmountToAdd = 0;
+			
+			void operator()()
+			{
+				TSharedPtr<FTestMessageSystem> StrongMessageSys = MessageSys.Pin();
+				check(StrongMessageSys.IsValid());
+
+				// It is expected that this instanced struct payload data here would go out of scope.
+				FTest_Payload_A PayloadData = {};
+				PayloadData.IncrementAmount = AmountToAdd;
+				FInstancedStruct PayloadDataInstance = FInstancedStruct::Make<FTest_Payload_A>(PayloadData);
+				
+				StrongMessageSys->QueueMessageForBroadcast(TestMessageId_A, PayloadDataInstance);
+			}
+		};
+
+		TArray<UE::Tasks::FTask> PendingTasks;
+		constexpr int32 NumTasksToSpawn = 5000;
+		for (int32 i = 0; i < NumTasksToSpawn; ++i)
+		{
+			UE::Tasks::FTask T = UE::Tasks::Launch(UE_SOURCE_LOCATION, FQueueMessageFunctor { .MessageSys = TestSystem, .AmountToAdd = i });
+			PendingTasks.Emplace(T);
+		}
+
+		// Wait for all the tasks which queue messages to complete
+		UE::Tasks::Wait(PendingTasks);
+
+		TestSystem->RunOnce();
+
+		TestSystem->Shutdown();
+		TestSystem.Reset();
+		
+		return true;
+	}
+
+	// A test which will spin up several tasks on different threads to attempt to bind listeners to messages
+	// on each thread. This will test that we can have different listeners binding on different threads
+	// for the same message system
+	IMPLEMENT_SIMPLE_AUTOMATION_TEST(TMessageSystemMultiThreadedBindListeners, "AsyncMessagePassing.MultiThreaded.BindListeners", QuickTestFlags)
+	bool TMessageSystemMultiThreadedBindListeners::RunTest(const FString& InParameters)
+	{
+		TSharedPtr<FTestMessageSystem, ESPMode::ThreadSafe> TestSystem = FAsyncMessageSystemBase::CreateMessageSystem<FTestMessageSystem>();
+		TestTrue(TEXT("MessageSystem startup has been called"), TestSystem->HasStarted());
+
+		// A functor to add a listener which will be run as a task
+		struct FAddListenerFunctor
+		{
+			TSharedPtr<FTestMessageSystem> MessageSys;
+
+			TSharedPtr<FTestListener> ListenerObj;
+			
+			void operator()()
+			{
+				check(MessageSys.IsValid());
+				
+				// Personally I am expecting this to fail until I lock up the listener map with a CS
+				const FAsyncMessageHandle Handle = MessageSys->BindListener(TestMessageId_A, ListenerObj.ToWeakPtr(), &FTestListener::Callback_Message_A);
+				check(Handle.IsValid());
+			}
+		};
+		
+		TArray<TSharedPtr<FTestListener>> Listeners;
+		
+		TArray<UE::Tasks::FTask> PendingTasks;
+		constexpr int32 NumTasksToSpawn = 5000;
+		for (int32 i = 0; i < NumTasksToSpawn; ++i)
+		{
+			TSharedPtr<FTestListener> Listener = MakeShared<FTestListener>();
+			Listeners.Emplace(Listener);
+			
+			UE::Tasks::FTask T = UE::Tasks::Launch(UE_SOURCE_LOCATION, FAddListenerFunctor { .MessageSys = TestSystem, .ListenerObj = Listener });
+			PendingTasks.Emplace(T);
+		}
+		
+		// Wait for all the tasks which queue messages to complete
+		UE::Tasks::Wait(PendingTasks);
+
+		TestSystem->RunOnce();
+		
+		TestSystem->Shutdown();
+		TestSystem.Reset();
+		
+		return true;
+	}
+
+	// A test which will spin up several tasks on different threads to attempt to bind listeners to messages
+	// on each thread. This will test that we can have different listeners binding on different threads
+	// for the same message system. The test system's handle ID is an atomic<uint32>, so we should have no issue
+	// generating handles from a bunch of different threads.
+	IMPLEMENT_SIMPLE_AUTOMATION_TEST(TMessageSystemMultiThreadedHandleGeneration, "AsyncMessagePassing.MultiThreaded.HandleGeneration", QuickTestFlags)
+	bool TMessageSystemMultiThreadedHandleGeneration::RunTest(const FString& InParameters)
+	{
+		TSharedPtr<FTestMessageSystem, ESPMode::ThreadSafe> TestSystem = FAsyncMessageSystemBase::CreateMessageSystem<FTestMessageSystem>();
+		TestTrue(TEXT("MessageSystem startup has been called"), TestSystem->HasStarted());
+
+		// A functor to add a listener which will be run as a task
+		struct FGenerateHandleFunctor
+		{
+			TSharedPtr<FTestMessageSystem> MessageSys;
+			
+			void operator()()
+			{
+				check(MessageSys.IsValid());
+				
+				// This should be valid because it is using atomics
+				const FAsyncMessageHandle Handle = MessageSys->GenerateTestHandle();
+			}
+		};
+		
+		TArray<UE::Tasks::FTask> PendingTasks;
+		constexpr int32 NumTasksToSpawn = 5000;
+		for (int32 i = 0; i < NumTasksToSpawn; ++i)
+		{
+			UE::Tasks::FTask T = UE::Tasks::Launch(UE_SOURCE_LOCATION, FGenerateHandleFunctor { .MessageSys = TestSystem });
+			PendingTasks.Emplace(T);
+		}
+
+		// Wait for all the tasks which queue messages to complete
+		UE::Tasks::Wait(PendingTasks);
+
+		TestSystem->RunOnce();
+		
+		TestSystem->Shutdown();
+		TestSystem.Reset();
+		
+		return true;
+	}
+
+	// Test to queue messages from multiple threads on a single message system
+	IMPLEMENT_SIMPLE_AUTOMATION_TEST(TMessageSystemMultiThreadedSequenceIdTest, "AsyncMessagePassing.MultiThreaded.UniqueSequenceId", QuickTestFlags)
+	bool TMessageSystemMultiThreadedSequenceIdTest::RunTest(const FString& InParameters)
+	{
+		TSharedPtr<FTestMessageSystem, ESPMode::ThreadSafe> TestSystem = FAsyncMessageSystemBase::CreateMessageSystem<FTestMessageSystem>();
+		TestTrue(TEXT("MessageSystem startup has been called"), TestSystem->HasStarted());
+
+		TSet<uint32> UsedSequenceIds;
+		uint32 NumTimesCalledback = 0;
+		
+		FAsyncMessageHandle ListenerHandle = TestSystem->BindListener(TestMessageId_A, [this, &NumTimesCalledback, &UsedSequenceIds](const FAsyncMessage& Message)
+		{
+			const uint32 SeqId = Message.GetSequenceId();
+			const bool bIsUnique = !UsedSequenceIds.Contains(SeqId);
+			TestTrue(TEXT("Has unique sequence ID"), bIsUnique);
+
+			UsedSequenceIds.Add(SeqId);
+			++NumTimesCalledback;
+		});
+		
+		struct FQueueMessageFunctor
+		{
+			TWeakPtr<FTestMessageSystem> MessageSys;
+
+			int32 AmountToAdd = 0;
+			
+			void operator()()
+			{
+				TSharedPtr<FTestMessageSystem> StrongMessageSys = MessageSys.Pin();
+				check(StrongMessageSys.IsValid());
+
+				// It is expected that this instanced struct payload data here would go out of scope.
+				FTest_Payload_A PayloadData = {};
+				PayloadData.IncrementAmount = AmountToAdd;
+				FInstancedStruct PayloadDataInstance = FInstancedStruct::Make<FTest_Payload_A>(PayloadData);
+				
+				StrongMessageSys->QueueMessageForBroadcast(TestMessageId_A, PayloadDataInstance);
+			}
+		};
+
+		TArray<UE::Tasks::FTask> PendingTasks;
+		constexpr int32 NumTasksToSpawn = 3000;
+		for (int32 i = 0; i < NumTasksToSpawn; ++i)
+		{
+			UE::Tasks::FTask T = UE::Tasks::Launch(UE_SOURCE_LOCATION, FQueueMessageFunctor { .MessageSys = TestSystem, .AmountToAdd = i });
+			PendingTasks.Emplace(T);
+		}
+
+		// Wait for all the tasks which queue messages to complete
+		UE::Tasks::Wait(PendingTasks);
+
+		TestSystem->RunOnce();
+
+		TestTrue(TEXT("Has the expected number of sequence IDs"),
+			NumTimesCalledback == NumTasksToSpawn &&
+			NumTasksToSpawn == UsedSequenceIds.Num());
+		
+		TestSystem->Shutdown();
+		TestSystem.Reset();
+		
+		return true;
+	}
+};
+
+#endif	// WITH_DEV_AUTOMATION_TESTS
+
+bool UE::AsyncMessageSystem::FTestMessageSystem::IsHandleBound(const FAsyncMessageHandle& Handle) const
+{
+	FScopeLock ListenerLock(&MessageListenerMapCS);
+	for (const TPair<FAsyncMessageId, FAsyncMessageBoundData>& MessagePair : BoundMessageListenerMap)
+	{
+		const FAsyncMessageBoundData& BoundData = MessagePair.Value;
+		for (const TPair<FAsyncMessageBindingOptions, TArray<FAsyncMessageIndividualListener<FMessageCallbackFunc>>>& ListenerPair : BoundData.ListenerMap)
+		{
+			const bool bHasHandle = ListenerPair.Value.ContainsByPredicate([Handle](const FAsyncMessageIndividualListener<FMessageCallbackFunc>& Listener)
+			{
+				return Listener.Handle == Handle;	
+			});
+			
+			if (bHasHandle)
+			{
+				return true;
+			}
+		}
+	}	
+	
+	return false;	
+}
+
+uint32 UE::AsyncMessageSystem::FTestMessageSystem::GetNumberOfListeners() const
+{
+	uint32 Res = 0;
+	// Check any message listeners...
+	{
+		FScopeLock ListenerLock(&MessageListenerMapCS);
+		for (const TPair<FAsyncMessageId, FAsyncMessageBoundData>& MessagePair : BoundMessageListenerMap)
+		{
+			const FAsyncMessageBoundData& BoundData = MessagePair.Value;
+			for (const TPair<FAsyncMessageBindingOptions, TArray<FAsyncMessageIndividualListener<FMessageCallbackFunc>>>& ListenerPair : BoundData.ListenerMap)
+			{
+				Res += ListenerPair.Value.Num();
+			}
+		}	
+	}
+
+	return Res;
+}
+
+FAsyncMessageHandle UE::AsyncMessageSystem::FTestMessageSystem::GenerateTestHandle()
+{
+	return GenerateNextValidMessageHandle({InternalTestTag_Invalid});
+}
+
+void UTestAsyncObject::CallbackFunction(const FAsyncMessage& Message)
+{
+	FConstStructView DataView = Message.GetPayloadView();
+	const FTest_Payload_A* Data = Message.GetPayloadData<const FTest_Payload_A>();
+	if (Data)
+	{
+		TestValue += Data->IncrementAmount;
+	}
+}
diff --git a/Engine/Plugins/Experimental/AsyncMessageSystem/Source/AsyncMessageSystem/Private/Tests/AsyncMessageSystemTests.h b/Engine/Plugins/Experimental/AsyncMessageSystem/Source/AsyncMessageSystem/Private/Tests/AsyncMessageSystemTests.h
new file mode 100644
--- /dev/null
+++ b/Engine/Plugins/Experimental/AsyncMessageSystem/Source/AsyncMessageSystem/Private/Tests/AsyncMessageSystemTests.h
@@ -0,0 +1,95 @@
+// Copyright Epic Games, Inc. All Rights Reserved.
+
+#pragma once
+
+#include "CoreMinimal.h"
+#include "UObject/Object.h"
+#include "AsyncGameplayMessageSystem.h"
+#include "AsyncMessageSystemTests.generated.h"
+
+USTRUCT()
+struct FTest_Payload_A
+{
+	GENERATED_BODY()
+
+	UPROPERTY()
+	int32 IncrementAmount = 78;
+
+	UPROPERTY()
+	FString SomeName = TEXT("Test string");
+
+	UPROPERTY()
+	float Bar = 123.0f;
+};
+
+struct FAsyncMessage;
+
+UCLASS(NotBlueprintable)
+class UTestAsyncObject : public UObject
+{
+	GENERATED_BODY()
+
+public:
+	
+	void CallbackFunction(const FAsyncMessage& Message);
+
+	UPROPERTY()
+	int32 TestValue = 5;
+	
+};
+
+namespace UE::AsyncMessageSystem
+{
+	// Test message system interface impl
+	class FTestMessageSystem : public FAsyncMessageSystemBase
+	{
+	protected:
+		virtual void Startup_Impl() override
+		{
+			bHasStarted = true;
+		}
+
+		virtual void Shutdown_Impl() override
+		{
+			bHasShutdown = true;
+		}
+
+		virtual void PostQueueMessage(const FAsyncMessageId MessageId, const TSet<FAsyncMessageBindingOptions>& OptionsBoundTo) override
+		{
+			// Dont do anything by default here in this simple test message system
+		}
+		
+		bool bHasStarted = false;
+		bool bHasShutdown = false;	
+
+	public:
+		const bool HasStarted() const { return bHasStarted; }
+		const bool HasShutdown() const { return bHasShutdown && bIsShuttingDown; }
+		
+		/**
+		 * Returns true if the handle currently has any listeners bound to it in the message map
+		 */
+		bool IsHandleBound(const FAsyncMessageHandle& Handle) const;
+
+		uint32 GetNumberOfListeners() const;
+		
+		void RunOnce()
+		{
+			ProcessMessageQueueForBinding(FAsyncMessageBindingOptions{});
+		}
+
+		void Test_Shutdown()
+		{
+			Shutdown();
+		}
+
+		FAsyncMessageHandle GenerateTestHandle();
+	};
+}
+
+USTRUCT()
+struct FNested_Payload
+{
+	GENERATED_BODY()
+	TWeakPtr<UE::AsyncMessageSystem::FTestMessageSystem> MessageSystem = nullptr;
+};
\ No newline at end of file
diff --git a/Engine/Plugins/Experimental/AsyncMessageSystem/Source/AsyncMessageSystem/Public/AsyncGameplayMessageSystem.h b/Engine/Plugins/Experimental/AsyncMessageSystem/Source/AsyncMessageSystem/Public/AsyncGameplayMessageSystem.h
new file mode 100644
--- /dev/null
+++ b/Engine/Plugins/Experimental/AsyncMessageSystem/Source/AsyncMessageSystem/Public/AsyncGameplayMessageSystem.h
@@ -0,0 +1,86 @@
+// Copyright Epic Games, Inc. All Rights Reserved.
+
+#pragma once
+
+#include "AsyncMessageSystemBase.h"
+#include "Engine/EngineBaseTypes.h"		// For ELevelTick
+#include "Async/TaskGraphFwd.h"			// FGraphEventRef
+
+class UWorld;
+
+/**
+ * Implementation of an async message system which scheudles the processing of messages based on tick groups
+ * and named thread async tasks, making it integrate nicely with common gameplay frameworks in Unreal.
+ */
+class ASYNCMESSAGESYSTEM_API FAsyncGameplayMessageSystem final :
+	public FAsyncMessageSystemBase
+{
+	friend struct FMessageSystemTickFunction;
+
+public:
+	explicit FAsyncGameplayMessageSystem(UWorld* OwningWorld);
+
+private:
+	//~ Begin FAsyncMessageSystemBase interface
+	virtual void Startup_Impl() override;
+	virtual void Shutdown_Impl() override;
+	virtual void PostQueueMessage(const FAsyncMessageId MessageId, const TSet<FAsyncMessageBindingOptions>& OptionsBoundTo) override;
+	//~ End FAsyncMessageSystemBase interface
+
+	/**
+	 * Creates and registers tick functions for this message system to the outer world.
+	 */
+	void CreateTickFunctions();
+
+	/**
+	 * Destroy any registered tick functions on this message system
+	 */
+	void DestroyTickFunctions();
+
+	/**
+	 * Starts an async task to process the messages with listeners that have the given binding options.
+	 * This should ONLY be called for binding options which are not tick groups.
+	 * 
+	 * @param Options The options of which to start async processing for.
+	 */
+	void StartAsyncProcessForBinding(const FAsyncMessageBindingOptions& Options);
+	
+	/**
+	 * Executes a single tick of this message system for a given group 
+	 */
+	void ExecuteTick(
+		float DeltaTime,
+		ELevelTick TickType,
+		ENamedThreads::Type CurrentThread,
+		const FGraphEventRef& MyCompletionGraphEvent,
+		ETickingGroup TickGroup);
+
+	/**
+	 * The world in which this message system belongs to. The tick functions of this message system will
+	 * be registered to this world's persistent level.
+	 */
+	TWeakObjectPtr<UWorld> OuterWorld = nullptr;
+
+	/**
+	 * Tick functions which are created when this message system starts. There is one tick functions
+	 * for each tick group that this message system supports and they will depend on one another
+	 * to keep a deterministic ordering of ticking on this message system.
+	 */
+	TArray<TSharedPtr<FTickFunction>> TickFunctions;
+
+	/***
+	 * The tick group earliest in the frame which this message system supports.
+	 */
+	static constexpr ETickingGroup EarliestSupportedTickGroup = TG_PrePhysics;
+	
+	/**
+	 * The tick group that is the latest in frame time which this message system will support
+	 */
+	static constexpr ETickingGroup LatestSupportedTickGroup = TG_PostUpdateWork;
+
+	/**
+	 * Keep track of the state of what tick group is currently executing and what the last one was.
+	 */
+	ETickingGroup CurrentTickGroup = EarliestSupportedTickGroup;
+	ETickingGroup LastTickedGroup = LatestSupportedTickGroup;
+};
diff --git a/Engine/Plugins/Experimental/AsyncMessageSystem/Source/AsyncMessageSystem/Public/AsyncMessage.h b/Engine/Plugins/Experimental/AsyncMessageSystem/Source/AsyncMessageSystem/Public/AsyncMessage.h
new file mode 100644
--- /dev/null
+++ b/Engine/Plugins/Experimental/AsyncMessageSystem/Source/AsyncMessageSystem/Public/AsyncMessage.h
@@ -0,0 +1,144 @@
+// Copyright Epic Games, Inc. All Rights Reserved.
+
+#pragma once
+
+#include "HAL/Platform.h"	// For uint64 type decl
+#include "AsyncMessageId.h"
+#include "StructUtils/StructView.h"
+
+/**
+ * A single async message which has been queued for broadcasting, and listeners are
+ * interested in receiving the data from.
+ *
+ * Async Message Messages contain some data about when, where, and how the message was queued
+ * which make it possible to synconize game state across multiple threads with messages. 
+ */
+struct ASYNCMESSAGESYSTEM_API FAsyncMessage final
+{
+public:
+	
+	FAsyncMessage() = default;
+	
+	FAsyncMessage(
+		const FAsyncMessageId& MessageId,
+		const FAsyncMessageId& MessageSourceId,
+		const double MessageTimestamp,
+		const uint64 CurrentFrame,
+		const uint32 ThreadQueuedFrom,
+		const uint32 MessageSequenceCounter,
+		FConstStructView InPayloadData);
+
+	/**
+	 * Returns this message’s unique ID which listeners use to bind to it. 
+	 */
+	FAsyncMessageId GetMessageId() const;
+
+	/**
+	 * Sets the current message id which this FAsyncMessage represents to another message
+	 * 
+	 * @param NewMessageId The new message Id to set
+	 */
+	void SetMessageId(const FAsyncMessageId& NewMessageId);
+
+	/**
+	 * Returns this the source message which caused this message to be queued
+	 */
+	FAsyncMessageId GetMessageSourceId() const;
+	
+	/**
+	 * Returns the timestamp which this message message was queued on.
+	 */
+	double GetQueueTimestamp() const;
+
+	/**
+	 * Returns the frame (GFrameNumber) of which this message was queued for broadcast.
+	 */
+	uint64 GetQueueFrame() const;
+
+	/**
+	 * Returns the thread Id from which this message was queued from ( FPlatformTLS::GetCurrentThreadId() )
+	 */
+	uint32 GetThreadQueuedFromThreadId() const;
+	
+	/**
+	 * Returns the “sequence ID” of this message.
+	 * If there are multiple of the same message queued for broadcast on a single frame,
+	 * then this sequence ID can be used to order them and determine which was queued first if desired.	 
+	 */
+	uint32 GetSequenceId() const;
+
+	/**
+	 * Get the payload data associated with this async message.
+	 */
+	FStructView GetPayloadView() const;
+
+	/**
+	 * Returns the typed data pointer of this message's payload. Nullptr if the payload is invalid.
+	 */
+	template<typename TDataType>
+	TDataType* GetPayloadData() const
+	{
+		return PayloadCopy.GetPtr<TDataType>();
+	}
+
+private:
+	
+	/**
+	 * The Async message Id which this message is a part of. 
+	 */
+	FAsyncMessageId MessageId = FAsyncMessageId::Invalid;
+
+	/**
+	 * The source of this message. This might be different from the "MessageId", because this could be a parent message which
+	 * is being triggered by a child message. In that case, the "MessageId" would be the parent message, and the "MessageSourceId" would be the
+	 * actual child message which had triggered.
+	 *
+	 * For example, if "message.colors.red" is a child message of "message.colors", then the "MessageId" would be "message.colors" and
+	 * the MessageSourceId would be "message.colors.red".
+	 *
+	 * If the MessageSourceId and MessageId are the same value, then that means that this is the childmost message.
+	 */
+	FAsyncMessageId MessageSourceId = FAsyncMessageId::Invalid;
+
+	/**
+	 * The time at which this message was queued ( FApp::GetCurrentTime() )
+	 */
+	double QueueTime = 0.0;
+
+	/**
+	 * The frame at which this message was queued ( GFrameCounter )
+	 */
+	uint64 QueueFrame = 0;
+
+	/**
+	 * The ThreadId from which this message was queued from ( FPlatformTLS::GetCurrentThreadId() )
+	 */
+	uint32 ThreadQueuedFrom = 0;
+
+	/**
+	 * The sequence of this specific message if there were multiple in the message queue at one time.
+	 * 
+	 * This will allow async listeners to determine which message of this type was queued first if
+	 * there are multiple.
+	 */
+	uint32 SequenceId = 0;
+
+	/**
+	 * A copy of the payload data that is created upon construction of this message.
+	 * 
+	 * This data is copied in the constructor of this message when the message is queued
+	 * in order to make the data thread safe to access.
+	 *
+	 * If you would like to have more control over the payload data and how it is synconized
+	 * (for example if you have a costly payload type which would be too expensive to copy)
+	 * then you could make your payload struct contain a simple pointer and handle synchronization
+	 * yourself with critical sections. 
+	 * 
+	 * TODO: When we copy the payload data, have it use a custom linear allocator on the owning message
+	 * system to make it faster to copy and more cache friendly to the message system to iterate messages.
+	 */
+	FInstancedStruct PayloadCopy;
+
+	// TODO: We could perhaps define a "lifetime" concept on these messages, similar to how AnimNext's
+	// TraitEvents (TraitEvent.h) work, as well as consumption perhaps?	
+};
\ No newline at end of file
diff --git a/Engine/Plugins/Experimental/AsyncMessageSystem/Source/AsyncMessageSystem/Public/AsyncMessageBindingOptions.h b/Engine/Plugins/Experimental/AsyncMessageSystem/Source/AsyncMessageSystem/Public/AsyncMessageBindingOptions.h
new file mode 100644
--- /dev/null
+++ b/Engine/Plugins/Experimental/AsyncMessageSystem/Source/AsyncMessageSystem/Public/AsyncMessageBindingOptions.h
@@ -0,0 +1,48 @@
+// Copyright Epic Games, Inc. All Rights Reserved.
+
+#pragma once
+
+#include "Engine/EngineBaseTypes.h"			// For ETickGround
+#include "Async/TaskGraphInterfaces.h"		// for ENamedThreads
+
+/**
+ * Options used to specify when a binder would like to receive a message when it is broadcast.
+ */
+struct ASYNCMESSAGESYSTEM_API FAsyncMessageBindingOptions final
+{
+	FAsyncMessageBindingOptions() = default;
+	
+	/**
+	 * Determines which type of configuration this listener follows
+	 */
+	enum class EBindingType : uint8
+	{
+		/**
+		 * If set, then this binding will receive a message on the main game thread during the specific tick group.
+		 */
+		UseTickGroup = 0,
+
+		/**
+		 * If set, then this binding will receive a message on this TaskId's thread at some point
+		 */
+		UseTaskId = 1,
+
+		// TODO: UE::Tasks priority and extended priority groups
+	};
+
+	EBindingType Type = EBindingType::UseTickGroup;
+
+	/**
+	 * Identifies the tick group which the binding object would like to receive a callback in
+	 */
+	ETickingGroup Group = ETickingGroup::TG_PostUpdateWork;
+		
+	/**
+	 * Identifies the thread ID in which the binding object would like to receive a callback on
+	 */
+	int32 TaskId = ENamedThreads::Type::GameThread;
+
+	bool operator==(const FAsyncMessageBindingOptions& Other) const;
+	bool operator!=(const FAsyncMessageBindingOptions& Other) const;
+	ASYNCMESSAGESYSTEM_API friend uint32 GetTypeHash(const FAsyncMessageBindingOptions& Opts);
+};
\ No newline at end of file
diff --git a/Engine/Plugins/Experimental/AsyncMessageSystem/Source/AsyncMessageSystem/Public/AsyncMessageHandle.h b/Engine/Plugins/Experimental/AsyncMessageSystem/Source/AsyncMessageSystem/Public/AsyncMessageHandle.h
new file mode 100644
--- /dev/null
+++ b/Engine/Plugins/Experimental/AsyncMessageSystem/Source/AsyncMessageSystem/Public/AsyncMessageHandle.h
@@ -0,0 +1,83 @@
+// Copyright Epic Games, Inc. All Rights Reserved.
+
+#pragma once
+
+#include "Containers/UnrealString.h"	// For FString
+#include "HAL/Platform.h"	// For uint32 type decl
+#include "AsyncMessageId.h"
+
+/**
+ * Handle used to identify a bound listener to an Async Message.
+ * These handles are unique to each bound listener and created by the
+ * owning FAsyncMessageSystemBase that the Message was bound to.
+ */
+struct ASYNCMESSAGESYSTEM_API FAsyncMessageHandle final
+{
+public:
+
+	friend class FAsyncMessageSystemBase;
+
+	/**
+	 * Default constructor of a Message handle. This default handle will be invalid.
+	 * Valid handles should only be created from an operating async message system (@see IAsyncMessageSystem).
+	 */
+	FAsyncMessageHandle() = default;
+
+	/**
+	 * Returns true if this Message handle has a valid internal ID.
+	 */
+	bool IsValid() const;
+
+	/**
+	 * Returns the internal ID of this handle.
+	 */
+	uint32 GetId() const;
+
+	/**
+	 * Returns the Message in which this handle is bound to.
+	 */
+	FAsyncMessageId GetBoundMessageId() const;
+
+	/**
+	 * Returns a formatted string which represents this handle's data.
+	 */
+	FString ToString() const;
+
+	bool operator==(const FAsyncMessageHandle& Other) const;
+	bool operator!=(const FAsyncMessageHandle& Other) const;
+	bool operator>=(const FAsyncMessageHandle& Other) const;
+	bool operator<(const FAsyncMessageHandle& Other) const;
+	ASYNCMESSAGESYSTEM_API friend uint32 GetTypeHash(const FAsyncMessageHandle& InMapping);
+
+	/**
+	 * Represents an invalid handle which cannot have any Message listeners bound to it
+	 * and means that no more Messages will be received. 
+	 */
+	static const FAsyncMessageHandle Invalid;
+
+private:
+
+	/**
+	 * Private constructor for use by the message system
+	 */
+	FAsyncMessageHandle(const uint32 InHandleValue, const FAsyncMessageId BoundMessageId);
+
+	/**
+	 * Internal index for message handles which we consider to be invalid.
+	 */
+	static constexpr uint32 InvalidHandleIndex = 0;
+	
+	/**
+	 * The value of this internal handle.
+	 */
+	uint32 InternalHandle = FAsyncMessageHandle::InvalidHandleIndex;
+
+	/**
+	 * The async message that this handle is bound to.
+	 */
+	FAsyncMessageId BoundMessage = FAsyncMessageId::Invalid;
+
+	// TODO: Should these handles also have a pointer back to the message system that they
+	// were generated from? That might be useful in the context of having multiple different
+	// message systems perhaps...
+};
\ No newline at end of file
diff --git a/Engine/Plugins/Experimental/AsyncMessageSystem/Source/AsyncMessageSystem/Public/AsyncMessageId.h b/Engine/Plugins/Experimental/AsyncMessageSystem/Source/AsyncMessageSystem/Public/AsyncMessageId.h
new file mode 100644
--- /dev/null
+++ b/Engine/Plugins/Experimental/AsyncMessageSystem/Source/AsyncMessageSystem/Public/AsyncMessageId.h
@@ -0,0 +1,93 @@
+// Copyright Epic Games, Inc. All Rights Reserved.
+
+#pragma once
+
+#include "UObject/NameTypes.h"			// For FName
+#include "Containers/UnrealString.h"	// For FString
+#include "GameplayTagContainer.h"		// For FGameplayTag
+
+#include "AsyncMessageId.generated.h"
+
+/**
+ * Represents a single message which can be bound to and queued in the async message system.
+ *
+ * This ID is what you use to bind to (listen for) messages which are being broadcast, as well as
+ * actually queue the message for broadcasting.
+ *
+ * Internally this is simply a FGameplayTag, which is how this Message's parent hierarchy is determinted.  
+ */
+USTRUCT(BlueprintType)
+struct ASYNCMESSAGESYSTEM_API FAsyncMessageId final
+{
+	GENERATED_BODY()
+
+public:
+	FAsyncMessageId() = default;
+
+	/**
+	 * Constructs a MessageId from the given FName.
+	 *
+	 * Note: This will internally be requesting a gameplay tag representation of this Fname, and will ensure if it is not found.
+	 * 
+	 * @param MessageName The FName representation of this MessageId
+	 */
+	FAsyncMessageId(const FName MessageName);
+	
+	/**
+	 * Constructs a FAsyncMessageId from the FName assocaited with the given Gameplay Tag. 
+	 * 
+	 * @param MessageTag The gameplay tag to construct a Message Id from
+	 */
+	FAsyncMessageId(const FGameplayTag& MessageTag);
+	~FAsyncMessageId() = default;
+
+	/**
+	 * @return True if this is a valid MessageId which can be used to bind and queue messages.
+	 */
+	[[nodiscard]] bool IsValid() const;
+
+	/**
+	 * @return The raw FName which represents this MessageId 
+	 */
+	[[nodiscard]] FName GetMessageName() const;
+
+	/**
+	 * @return FString representation of this MessageId 
+	 */
+	[[nodiscard]] FString ToString() const;
+
+	/**
+	* Returns the parent of this async message. The parent is parsed out of the underlying FName which represents this message by '.' delimeters.
+	* 
+	* For example, the parent of "message.id.test" is "message.id" , and the parent of "message.id" is "message". If there are no '.' delimeters in the
+	* message's name then it does not have a parent, and this will return FAsyncMessageId::Invalid
+	*/
+	[[nodiscard]] FAsyncMessageId GetParentMessageId() const;	
+
+	bool operator==(const FAsyncMessageId& Other) const;
+	bool operator!=(const FAsyncMessageId& Other) const;
+	ASYNCMESSAGESYSTEM_API friend uint32 GetTypeHash(const FAsyncMessageId& InMapping);
+	
+	/**
+	 * Represents an invalid message ID which cannot be bound to.
+	 */
+	static const FAsyncMessageId Invalid;
+
+	/**
+	 * Walks up the parent heirarchy of the given Startingmessage. This will iterate up the parent message chain
+	 * and calling the given ForEachMessageFunc for it so long as it is valid.
+	 *
+	 * This does include the given starting message!
+	 */
+	static void WalkMessageHierarchy(const FAsyncMessageId StartingMessage, TFunction<void(const FAsyncMessageId MessageId)> ForEachMessageFunc);
+	
+protected:
+	
+	/**
+	 * The FName which represents this async message ID
+	 */
+	UPROPERTY(EditAnywhere, BlueprintReadWrite, Category = "Messages")
+	FGameplayTag InternalMessageTag;
+};
+
+
diff --git a/Engine/Plugins/Experimental/AsyncMessageSystem/Source/AsyncMessageSystem/Public/AsyncMessageSystemBase.h b/Engine/Plugins/Experimental/AsyncMessageSystem/Source/AsyncMessageSystem/Public/AsyncMessageSystemBase.h
new file mode 100644
--- /dev/null
+++ b/Engine/Plugins/Experimental/AsyncMessageSystem/Source/AsyncMessageSystem/Public/AsyncMessageSystemBase.h
@@ -0,0 +1,415 @@
+// Copyright Epic Games, Inc. All Rights Reserved.
+
+#pragma once
+
+#include "AsyncMessageHandle.h"
+#include "AsyncMessageId.h"
+#include "AsyncMessage.h"
+#include "AsyncMessageStore.h"
+#include <atomic>
+#include "Containers/Map.h"
+#include "Containers/MpscQueue.h"
+#include "StructUtils/StructView.h"
+#include "Templates/SharedPointer.h"
+#include "UObject/WeakObjectPtrTemplates.h"
+#include "UObject/StrongObjectPtr.h"
+
+/**
+ * Base abstract implementation of an async message system which can be used to easily
+ * pass signals and messages across different threads in Unreal.
+ */
+class ASYNCMESSAGESYSTEM_API FAsyncMessageSystemBase :
+	public TSharedFromThis<FAsyncMessageSystemBase>
+{
+public:
+
+	/**
+	 * Creates a new shared pointer to a message system of type TMessageSystemType
+	 * and calls the "Startup" function on it.
+	 *
+	 * Make sure to call the "Shutdown" function on the message system prior to its destruction.
+	 */
+	template<class TMessageSystemType, typename... TArgs>
+	[[nodiscard]] static inline TSharedPtr<TMessageSystemType, ESPMode::ThreadSafe> CreateMessageSystem(TArgs&&... Args);
+
+	using FMessageCallbackFunc = TFunction<void(const FAsyncMessage&)>;
+	
+protected:
+	
+	// We have a protected constructor on the message system because you should always use the
+	// "CreateMessageSystem" function above to make one. This allows us to ensure the virtual
+	// "Startup" function is called.
+	FAsyncMessageSystemBase() = default;
+
+public:
+
+	UE_NONCOPYABLE(FAsyncMessageSystemBase)
+	virtual ~FAsyncMessageSystemBase();
+
+	/**
+	 * Binds the given callback to this message so that when an message of type MessageId is broadcast, it will be executed.
+	 *
+	 * @param MessageId		The async message to bind to
+	 * @param Callback		The callback which should be executed when 'MessageId' is broadcast.
+	 * @param Options		Options describing when the listener would like to receive this callback (what tick group or named async thread)
+	 *
+	 * @return Async Message handle which can be used to unbind this listener. If the listener is not unbound, then a weak reference to it will be checked
+	 * and cleaned up when this message system is shut down. If the listener goes out of scope and the callback function is no longer valid, it will be unbound.
+	 */
+	FAsyncMessageHandle BindListener(
+		const FAsyncMessageId MessageId,
+		FMessageCallbackFunc&& Callback,
+		const FAsyncMessageBindingOptions& Options = {});
+
+	/**
+	 * Binds the given callback to this message so that when an message of type MessageId is broadcast, it will be executed.
+	 *
+	 * @param MessageId			The async message to bind to
+	 * @param WeakOwnerObject	The object that owns the bound listener callback. If this object is not valid at the time this message is called, then nothing will happen
+	 *							(the Callback delegate will not be execute on it)
+	 *							
+	 * @param Callback			The callback which should be executed when 'MessageId' is broadcast.
+	 * @param Options		Options describing when the listener would like to receive this callback (what tick group or named async thread)
+	 * 
+	 * @return Async Message handle which can be used to unbind this listener. If the listener is not unbound, then a weak reference to it will be checked
+	 * and cleaned up when this message system is shut down. If the listener goes out of scope and the callback function is no longer valid, it will be unbound.
+	 */
+	template <typename TOwner = UObject>
+	FAsyncMessageHandle BindListener(
+		const FAsyncMessageId MessageId,
+		TWeakObjectPtr<TOwner> WeakOwnerObject,
+		void(TOwner::* Callback)(const FAsyncMessage&),
+		const FAsyncMessageBindingOptions& Options = {});
+	
+	/**
+	 * Binds the given callback to this message so that when an message of type MessageId is broadcast, it will be executed.
+	 *
+	 * @param MessageId		The async message to bind to
+	 * @param WeakObject	The object that owns the bound listener callback. If this object is not valid at the time this message is called, then nothing will happen
+	 *						(the Callback delegate will not be execute on it)
+	 * 
+	 * @param Callback		The callback which should be executed when 'MessageId' is broadcast.
+	 * @param Options		Options describing when the listener would like to receive this callback (what tick group or named async thread)
+	 *
+	 * @return Async message handle which can be used to unbind this listener. If the listener is not unbound, then a weak reference to it will be checked
+	 * and cleaned up when this message system is shut down. If the listener goes out of scope and the callback function is no longer valid, it will be unbound.
+	 */
+	template <typename TOwner>
+	FAsyncMessageHandle BindListener(
+		const FAsyncMessageId MessageId,
+		TWeakPtr<TOwner> WeakObject,
+		void(TOwner::* Callback)(const FAsyncMessage&),
+		const FAsyncMessageBindingOptions& Options = {});
+
+	/**
+	 * Unbinds the given listener from its message so that it will no longer receive callbacks
+	 * when that message is broadcast.
+	 *
+	 * Listeners will be marked for unbinding and removed from the listener map the next time
+	 * that this system is processed.
+	 */
+	void UnbindListener(const FAsyncMessageHandle& HandleToUnbind);	
+	
+	/**
+	 * Queues the given async message for broadcast the next time that this
+	 * message system processes its message queue
+	 *
+	 * @param MessageId		The message that you would like to queue for broadcasting
+	 * @param PayloadData	Pointer to the payload data of this message. This payload data will be COPIED to the message queue to
+	 *						to make it safe for listeners on other threads. 
+	 *
+	 * @return True if this message had an listeners bound to it and it was successfully queued
+	 */
+	bool QueueMessageForBroadcast(const FAsyncMessageId MessageId, FConstStructView PayloadData = nullptr);
+
+	/**
+	 * Processes all async message on within the queue for the given binding options.
+	 * 
+	 * @param Options Which async message binding queue to process the messages from.
+	 */
+	void ProcessMessageQueueForBinding(const FAsyncMessageBindingOptions& Options);
+
+protected:
+	
+	/**
+	 * Provides an opportunity for this message system to run any startup logic which it may need to.
+	 * Utilize this to start any tick groups or other async background tasks which may need to occur.
+	 */
+	void Startup();
+
+public:
+	
+	/**
+	 * Manual shutdown of this message system. Once this is called, no more listeners can be bound and no more
+	 * messages can be queued. Use this to clean up any background work which may need to be done that was created in
+	 * Startup.
+	 */
+	void Shutdown();
+
+protected:
+
+	/**
+	* Provides an opportunity to run any startup logic that subclasses of this Base message
+	* system may need, such as creating tick groups for specific binding options.
+	*/
+	virtual void Startup_Impl() = 0;
+
+	/**
+	* Allows you to clean up anything you may need to with which was created during startup.
+	*/
+	virtual void Shutdown_Impl() = 0;
+
+	/**
+	* Handles binding the message callback to the given async message handle. This will return true if the binding was successful 
+	* and false if it failed for some reason (the message system could be shutting down, the MessageId could be invalid, etc.
+	*/
+	[[nodiscard]] virtual bool BindListener_Impl(const FAsyncMessageHandle& HandleToBindTo, const FAsyncMessageId MessageId, FMessageCallbackFunc&& Callback, const FAsyncMessageBindingOptions& Options);
+	
+	/**
+	 * Marks the given handle as ready to be unbound. Enqueues this handle to a queue which will be 
+	 * processed on the next call to ProcessUnbindHandleRequests.
+	 */
+	virtual void UnbindListener_Impl(const FAsyncMessageHandle& HandleToUnbind);
+
+	/**
+	* Actually queues the message for broadcasting.
+	* 
+	* Note: The payload data provided will be copied to make it safe for listeners on different threads. 
+	*/
+	virtual bool QueueMessageForBroadcast_Impl(const FAsyncMessageId MessageId, FConstStructView PayloadData);
+	
+	/**
+	 * Called after an message was queued (from QueueMessageForBroadcast_Impl). This forces child classes of the message system
+	 * to handle the queuing of messages and kick of some scheduling for actually processing the queues accordingly.
+	 *
+	 * @param MessageId			The message which was queued
+	 * @param OptionsBoundTo	All of the binding options that this message was queued for.
+	 */
+	virtual void PostQueueMessage(const FAsyncMessageId MessageId, const TSet<FAsyncMessageBindingOptions>& OptionsBoundTo) = 0;
+
+	/**
+	 * Actual implementation of how we process the message queue. This will dequeue all messages out fo the message store
+	 * for the given binding type, and broadcast their callbacks for each message and it's heirarhcy. 
+	 * 
+	 * @param Options The binding options which we are currently interested in processing messages for 
+	 */
+	virtual void ProcessMessageQueueForBinding_Impl(const FAsyncMessageBindingOptions& Options);
+	
+	/**
+	 * Allow an opportunity to pre-process any messages or bindings which may have occurred before 
+	 * the actual processing of the message queue
+	 */
+	virtual void PreProcessMessagesQueue(const FAsyncMessageBindingOptions& Options);
+
+	/**
+	 * Unbind any handles which have been flagged for removal in "UnbindListener"
+	 */
+	void ProcessUnbindHandleRequests();
+
+	/**
+	 * Iterate through the PendingBoundListenerQueue and bind the callback functions as necessary.
+	 * We need to defer the adding of listeners in order to support binding new listeners in response to
+	 * an message which is currently being processed (otherwise we would be modifying the listener array as we iterate it)
+	 */
+	void ProcessListenersPendingBinding();
+	
+	/**
+	 * Generates a new async message handle by atomically incrementing the NextMessageHandleId.
+	 * If the new value of NextMessageHandleId is above the max allowed handle id, it will be wrapped.
+	 */
+	[[nodiscard]] FAsyncMessageHandle GenerateNextValidMessageHandle(const FAsyncMessageId ForMessageId);
+	
+	/**
+	 * Critical section for when the message listeners map is changed (listeners are bound/unbound)
+	 */
+	mutable FCriticalSection MessageListenerMapCS;
+	
+	/**
+	 * The message handle which we should use to create next
+	 * New handles can be created any time that we bind a listener, which can be on any thread
+	 */
+	std::atomic<uint32> NextMessageHandleId = FAsyncMessageHandle::InvalidHandleIndex + 1u;
+
+	/**
+	 * The message sequence count. Each time an message is queued, we will increment
+	 * this. Doing this will allow you to determine if a message was queued for broadcast
+	 * before or after another of the same type or binding option, which can be useful if you are receiving
+	 * the same message multiple times on a single frame.
+	 *
+	 * This is reset back to zero every time we process a binding in ProcessMessageQueueForBinding_Impl. 
+	 */
+	std::atomic<uint32> NextMessageSequence = 0;
+
+	/**
+	 * A quick flag to check if this system is in the process of shutting down.
+	 */
+	std::atomic<bool> bIsShuttingDown = false;
+	
+	/**
+	 * Data stored about a single individual listener of a message
+	 */
+	template <typename CallbackFuncSignature>
+	struct FAsyncMessageIndividualListener
+	{
+		/**
+		* The handle to this individual listener assigned when it is bound to this message system.
+		*/
+		FAsyncMessageHandle Handle;
+		
+		/**
+		 * The delegate that listeners can bind to if they want to be called when this message
+		 * is broadcast. 
+		 */
+		CallbackFuncSignature Callback;
+	};
+	
+	/**
+	 * Data store for everything related to a single message
+	 */
+	struct FAsyncMessageBoundData
+	{
+		/**
+		 * A map of listeners for this message, ordered by their binding options.
+		 * 
+		 * This lets us keep track of how many binders are listening for this message on a specific
+		 * tick group or thread id
+		 */
+		TMap<FAsyncMessageBindingOptions, TArray<FAsyncMessageIndividualListener<FMessageCallbackFunc>>> ListenerMap;
+	};
+	
+	/**
+	 * Map of message ID's to their associated message data.
+	 */
+	TMap<FAsyncMessageId, FAsyncMessageBoundData> BoundMessageListenerMap;
+
+
+	/**
+	 * Where we store the messages when they are queued for broadcasting and where we can
+	 * process individual sets of FAsyncMessageBindingOptions for their messages 
+	 */
+	FAsyncMessageStore MessageStore;
+	
+	/**
+	* Stores data about a listener which has been bound but cannot be added at the time of which
+	* FAsyncMessageSystemBase::BindListener_Impl would be called, and needs to be deferred.
+	*/
+	template <typename TCallbackFuncSignature>
+	struct FPendingBoundListener
+	{
+		// The message which this listener desires to be bound to
+		const FAsyncMessageId MessageId;
+
+		// The data for this message, populated in FAsyncMessageSystemBase::BindListener_Impl
+		FAsyncMessageIndividualListener<TCallbackFuncSignature> Data;
+
+		// Binding options that this listener requires
+		FAsyncMessageBindingOptions Options = {};
+	};
+	
+	/**
+	 * A queue of listeners which are pending to be added. Listeners are added to this
+	 * queue during FAsyncMessageSystemBase::BindListener_Impl.
+	 * 
+	 * We do this because otherwise, we would be modifying the array of listeners whilst
+	 * iterating it, which would cause an error as well as some undefined behavior if you bind
+	 * another object to the same message. 
+	 */
+	TMpscQueue<FPendingBoundListener<FMessageCallbackFunc>> PendingBoundListenerQueue;
+
+	/**
+	 * message handles that are currently queued for removal. Handles get added to this queue
+	 * by UnbindListener_Impl, and the queue is processed in PreProcessMessagesQueue
+	 */
+	TMpscQueue<FAsyncMessageHandle> MessageHandlesPendingRemoval;		
+
+#if WITH_DEV_AUTOMATION_TESTS
+public:
+    /**
+     * A test-only helper function to quickly create an message handle from an index.
+     */
+    [[nodiscard]] static FAsyncMessageHandle GenerateHandleAtIndex(const uint32 Index, const FAsyncMessageId& ForId);
+#endif
+};
+
+//////////////////////////////////////////////////////////////
+// Template Definitions
+//////////////////////////////////////////////////////////////
+
+template <class TMessageSystemType, typename ... TArgs>
+TSharedPtr<TMessageSystemType, ESPMode::ThreadSafe> FAsyncMessageSystemBase::CreateMessageSystem(TArgs&&... Args)
+{
+	// Only accept subclasses of FAsyncMessageSystemBase here
+	static_assert(std::is_base_of<FAsyncMessageSystemBase, TMessageSystemType>::value, "TMessageSystemType not derived from FAsyncMessageSystemBase");
+	
+	TSharedPtr<TMessageSystemType> MessageSystem = MakeShared<TMessageSystemType>(Forward<TArgs>(Args)...);
+
+	MessageSystem->Startup();
+	
+	return MessageSystem;
+}
+
+template <typename TOwner>
+FAsyncMessageHandle FAsyncMessageSystemBase::BindListener(
+	const FAsyncMessageId MessageId,
+	TWeakObjectPtr<TOwner> WeakOwnerObject,
+	void(TOwner::* Callback)(const FAsyncMessage&),
+	const FAsyncMessageBindingOptions& Options)
+{
+	// This function is only for UObjects, because TWeakObjectPtr is a thing
+	static_assert(std::is_base_of<UObject, TOwner>::value, "TOwner not derived from UObject");
+
+	const FAsyncMessageHandle NewListenerHandle = GenerateNextValidMessageHandle(MessageId);
+	
+	const bool bWasSuccessfullyBound = BindListener_Impl(
+		NewListenerHandle,
+		MessageId,
+		[WeakOwnerObject, Callback, NewListenerHandle, WeakThisPtr = this->AsWeak()](const FAsyncMessage& Payload)
+		{
+			if (TStrongObjectPtr<TOwner> StrongObj = WeakOwnerObject.Pin())
+			{
+				(StrongObj.Get()->*Callback)(Payload);
+			}
+			// Otherwise, if the owning object which we have bound has gone out of scope, unbind its handle
+			else if (TSharedPtr<FAsyncMessageSystemBase> SharedPtr = WeakThisPtr.Pin())
+			{
+				SharedPtr->UnbindListener(NewListenerHandle);
+			}
+		},
+		Options);
+
+	return bWasSuccessfullyBound ? NewListenerHandle : FAsyncMessageHandle::Invalid;
+}
+
+template <typename TOwner>
+FAsyncMessageHandle FAsyncMessageSystemBase::BindListener(
+	const FAsyncMessageId MessageId,
+	TWeakPtr<TOwner> WeakObject,
+	void(TOwner::* Callback)(const FAsyncMessage&),
+	const FAsyncMessageBindingOptions& Options)
+{
+	const FAsyncMessageHandle NewListenerHandle = GenerateNextValidMessageHandle(MessageId);
+	
+	const bool bWasSuccessfullyBound = BindListener_Impl(
+		NewListenerHandle,
+		MessageId,
+		[WeakObject, Callback, NewListenerHandle, WeakThisPtr = this->AsWeak()](const FAsyncMessage& Payload)
+		{
+			if (TSharedPtr<TOwner> StrongObject = WeakObject.Pin())
+			{
+				if (TOwner* Obj = StrongObject.Get())
+				{
+					(Obj->*Callback)(Payload);
+					return;	// return early because we had a valid object
+				}
+			}
+
+			// The owning object pointer is invalid. Unbind it so that we don't attempt to call it again
+			if (TSharedPtr<FAsyncMessageSystemBase> SharedPtr = WeakThisPtr.Pin())
+			{
+				SharedPtr->UnbindListener(NewListenerHandle);
+			}
+		},
+		Options);
+
+	return bWasSuccessfullyBound ? NewListenerHandle : FAsyncMessageHandle::Invalid;
+}
diff --git a/Engine/Plugins/Experimental/AsyncMessageSystem/Source/AsyncMessageSystem/Public/AsyncMessageSystemLogs.h b/Engine/Plugins/Experimental/AsyncMessageSystem/Source/AsyncMessageSystem/Public/AsyncMessageSystemLogs.h
new file mode 100644
--- /dev/null
+++ b/Engine/Plugins/Experimental/AsyncMessageSystem/Source/AsyncMessageSystem/Public/AsyncMessageSystemLogs.h
@@ -0,0 +1,7 @@
+// Copyright Epic Games, Inc. All Rights Reserved.
+
+#pragma once
+
+#include "Logging/LogMacros.h"
+
+ASYNCMESSAGESYSTEM_API DECLARE_LOG_CATEGORY_EXTERN(LogAsyncMessageSystem, Log, All);
\ No newline at end of file
diff --git a/Engine/Plugins/Experimental/AsyncMessageSystem/Source/AsyncMessageSystem/Public/AsyncMessageWorldSubsystem.h b/Engine/Plugins/Experimental/AsyncMessageSystem/Source/AsyncMessageSystem/Public/AsyncMessageWorldSubsystem.h
new file mode 100644
--- /dev/null
+++ b/Engine/Plugins/Experimental/AsyncMessageSystem/Source/AsyncMessageSystem/Public/AsyncMessageWorldSubsystem.h
@@ -0,0 +1,69 @@
+// Copyright Epic Games, Inc. All Rights Reserved.
+
+#pragma once
+
+#include "Subsystems/WorldSubsystem.h"
+#include "Engine/World.h"
+
+#include "AsyncMessageWorldSubsystem.generated.h"
+
+class FAsyncMessageSystemBase;
+
+/**
+ * A world subsystem which will create a unique message system per-world
+ * and allow for easy access to the message system interface for gameplay code.
+ *
+ * An example of using the message system like this is:
+ *
+ * if (TSharedPtr<IAsyncMessageSystem> Sys = UAsyncMessageWorldSubsystem::GetSharedMessageSystem(GetWorld()))
+ * {
+ *		Sys->QueueMessageForBroadcast(MessageId, YourPayloadData);
+ * }
+ */
+UCLASS()
+class ASYNCMESSAGESYSTEM_API UAsyncMessageWorldSubsystem :
+	public UWorldSubsystem
+{
+
+GENERATED_BODY()
+
+public:
+
+	template<class TMessageSystemType = FAsyncMessageSystemBase>
+	static TSharedPtr<TMessageSystemType> GetSharedMessageSystem(const UWorld* InWorld)
+	{
+		static_assert(std::is_base_of<FAsyncMessageSystemBase, TMessageSystemType>::value, "TMessageSystemType not derived from FAsyncMessageSystemBase");
+
+		if (InWorld)
+		{
+			if (const UAsyncMessageWorldSubsystem* Sys = InWorld->GetSubsystem<UAsyncMessageWorldSubsystem>())
+			{
+				return Sys->GetSharedMessageSystem<TMessageSystemType>();	
+			}
+		}
+	
+		return nullptr;
+	}
+
+	template<class TMessageSystemType = FAsyncMessageSystemBase>
+	TSharedPtr<TMessageSystemType> GetSharedMessageSystem() const
+	{
+		static_assert(std::is_base_of<FAsyncMessageSystemBase, TMessageSystemType>::value, "TMessageSystemType not derived from FAsyncMessageSystemBase");
+
+		return StaticCastSharedPtr<TMessageSystemType>(MessageSystem);
+	}
+	
+protected:
+
+	virtual void Initialize(FSubsystemCollectionBase& Collection) override;
+	virtual void Deinitialize() override;
+
+	void InitMessageSystem();
+	void ShutdownMessageSystem();
+
+	/**
+	 * Instance of the message system for this world subsystem. This shares the lifetime of this subsystem
+	 * and will be released in ShutdownMessageSystem, upon Deinitialize of this subsystem.
+	 */
+	TSharedPtr<FAsyncMessageSystemBase, ESPMode::ThreadSafe> MessageSystem = nullptr;
+};
