Fix template deduction error

#rnx

diff --git a/Engine/Plugins/Experimental/AsyncMessageSystem/Source/AsyncMessageSystem/Private/AsyncMessageSystemBase.cpp b/Engine/Plugins/Experimental/AsyncMessageSystem/Source/AsyncMessageSystem/Private/AsyncMessageSystemBase.cpp
--- a/Engine/Plugins/Experimental/AsyncMessageSystem/Source/AsyncMessageSystem/Private/AsyncMessageSystemBase.cpp
+++ b/Engine/Plugins/Experimental/AsyncMessageSystem/Source/AsyncMessageSystem/Private/AsyncMessageSystemBase.cpp
@@ -107,7 +107,7 @@ bool FAsyncMessageSystemBase::BindListener_Impl(const FAsyncMessageHandle& Handl
 	// Let the message store know that this message will be listened for by this binding option
 	MessageStore.AddMessageToBinding(MessageId, Options);
 
-	FAsyncMessageIndividualListener ListenerData
+	FAsyncMessageIndividualListener<FMessageCallbackFunc> ListenerData
 		{
 			.Handle = HandleToBindTo,
 			.Callback = MoveTemp(Callback)
