Add support for binding to UE::Tasks priority types

diff --git a/Engine/Plugins/Experimental/AsyncMessageSystem/Source/AsyncMessageSystem/Private/AsyncGameplayMessageSystem.cpp b/Engine/Plugins/Experimental/AsyncMessageSystem/Source/AsyncMessageSystem/Private/AsyncGameplayMessageSystem.cpp
--- a/Engine/Plugins/Experimental/AsyncMessageSystem/Source/AsyncMessageSystem/Private/AsyncGameplayMessageSystem.cpp
+++ b/Engine/Plugins/Experimental/AsyncMessageSystem/Source/AsyncMessageSystem/Private/AsyncGameplayMessageSystem.cpp
@@ -5,6 +5,7 @@
 #include "Async/Async.h"
 #include "Engine/Level.h"
 #include "Engine/World.h"
+#include "Tasks/TaskPrivate.h"
 
 namespace UE::Private
 {
@@ -90,7 +91,7 @@ void FAsyncGameplayMessageSystem::PostQueueMessage(const FAsyncMessageId Message
 	// async task to process their message queue. 
 	for (const FAsyncMessageBindingOptions& BindingOpts : OptionsBoundTo)
 	{
-		if (BindingOpts.Type == FAsyncMessageBindingOptions::EBindingType::UseTaskId)
+		if (BindingOpts.Type == FAsyncMessageBindingOptions::EBindingType::UseTaskId || BindingOpts.Type == FAsyncMessageBindingOptions::EBindingType::UseTaskPriorities)
 		{
 			StartAsyncProcessForBinding(BindingOpts);
 		}
@@ -160,25 +161,46 @@ void FAsyncGameplayMessageSystem::StartAsyncProcessForBinding(const FAsyncMessag
 {
 	TRACE_CPUPROFILER_EVENT_SCOPE(FMessageSystem::StartAsyncProcessForBinding);
 
-	// We only want to do this for task ID's, tick groups are already being processed by us.
-	check(Options.Type == FAsyncMessageBindingOptions::EBindingType::UseTaskId);
+	// We only want to do this for task ID's and priorities. Tick groups are already being processed via our tick functions
+	check(
+		Options.Type == FAsyncMessageBindingOptions::EBindingType::UseTaskId ||
+		Options.Type == FAsyncMessageBindingOptions::EBindingType::UseTaskPriorities);
 
 	// Kick off a weak lambda to process the message queue for this task graph ID, ensuring
 	// that any listeners bound to these options will get the message called back when they expect
 	TWeakPtr<FAsyncGameplayMessageSystem> WeakThisPtr = StaticCastWeakPtr<FAsyncGameplayMessageSystem>(this->AsWeak());
-
 	
-	const ENamedThreads::Type ThreadToProcessOn = static_cast<ENamedThreads::Type>(Options.TaskId); 
+	UE::Tasks::ETaskPriority TaskPri = Options.TaskPriority;
+	UE::Tasks::EExtendedTaskPriority ExtendedTaskPri = Options.ExtendedTaskPriority;
 
-		// TODO: This could be a UE::Tasks::Launch call, but we would need to convert the ENamedThreads to ETaskPriority and EExtendedTaskPriority.
-    	// This conversion only exists in one direction right now and is not public, see ENamedThreads::Type TranslatePriority.
-	AsyncTask(ThreadToProcessOn, [WeakThisPtr, Options]()
+	// Translate from the old named thread model to the newer UE Tasks model if we need to
+	if (Options.Type == FAsyncMessageBindingOptions::EBindingType::UseTaskId)
 	{
-		if (TSharedPtr<FAsyncGameplayMessageSystem> This = WeakThisPtr.Pin())
-		{
-			This->ProcessMessageQueueForBinding(Options);
-		}
-	});
+		const ENamedThreads::Type ThreadToProcessOn = static_cast<ENamedThreads::Type>(Options.TaskId);
+		UE::Tasks::Private::TranslatePriority(ThreadToProcessOn, OUT TaskPri, OUT ExtendedTaskPri);
+	}
+	
+	UE::Tasks::Launch(UE_SOURCE_LOCATION, [WeakThisPtr, Options]()
+    {
+    	if (TSharedPtr<FAsyncGameplayMessageSystem> This = WeakThisPtr.Pin())
+    	{
+    		This->ProcessMessageQueueForBinding(Options);
+    	}
+    },
+    TaskPri,
+    ExtendedTaskPri);
+	
+	// The above UE::Tasks::Launch is the same thing as this older Task Graph syntax:
+	// AsyncTask(ThreadToProcessOn, [WeakThisPtr, Options]()
+	// {
+	// 	if (TSharedPtr<FAsyncGameplayMessageSystem> This = WeakThisPtr.Pin())
+	// 	{
+	// 		This->ProcessMessageQueueForBinding(Options);
+	// 	}
+	// });
+	//
+	// We utilize the UE::Tasks system here and translate the ENamedThread because the UE::Tasks system
+	// should have less overhead and better schdueling capabilities in most scenarios
 }
 
 void FAsyncGameplayMessageSystem::ExecuteTick(
diff --git a/Engine/Plugins/Experimental/AsyncMessageSystem/Source/AsyncMessageSystem/Private/AsyncMessageBindingOptions.cpp b/Engine/Plugins/Experimental/AsyncMessageSystem/Source/AsyncMessageSystem/Private/AsyncMessageBindingOptions.cpp
--- a/Engine/Plugins/Experimental/AsyncMessageSystem/Source/AsyncMessageSystem/Private/AsyncMessageBindingOptions.cpp
+++ b/Engine/Plugins/Experimental/AsyncMessageSystem/Source/AsyncMessageSystem/Private/AsyncMessageBindingOptions.cpp
@@ -2,12 +2,30 @@
 
 #include "AsyncMessageBindingOptions.h"
 
+FAsyncMessageBindingOptions::FAsyncMessageBindingOptions(const ETickingGroup DesiredTickGroup)
+	: Type(EBindingType::UseTickGroup)
+	, Group(DesiredTickGroup)
+{ }
+
+FAsyncMessageBindingOptions::FAsyncMessageBindingOptions(const ENamedThreads::Type NamedThread)
+	: Type(EBindingType::UseTaskId)
+	, TaskId(NamedThread)
+{ }
+
+FAsyncMessageBindingOptions::FAsyncMessageBindingOptions(const UE::Tasks::ETaskPriority InTaskPriority, const UE::Tasks::EExtendedTaskPriority InExtendedTaskPriority)
+	: Type(EBindingType::UseTaskPriorities)
+	, TaskPriority(InTaskPriority)
+	, ExtendedTaskPriority(InExtendedTaskPriority)
+{ }
+
 bool FAsyncMessageBindingOptions::operator==(const FAsyncMessageBindingOptions& Other) const
 {
 	return
 		Type == Other.Type &&
 		Group == Other.Group &&
-		TaskId == Other.TaskId;
+		TaskId == Other.TaskId &&
+		TaskPriority == Other.TaskPriority && 
+		ExtendedTaskPriority == Other.ExtendedTaskPriority;
 }
 
 bool FAsyncMessageBindingOptions::operator!=(const FAsyncMessageBindingOptions& Other) const
@@ -22,6 +40,8 @@ uint32 GetTypeHash(const FAsyncMessageBindingOptions& Opts)
 	Hash = HashCombine(Hash, GetTypeHash(static_cast<uint8>(Opts.Type)));
 	Hash = HashCombine(Hash, GetTypeHash(static_cast<int32>(Opts.Group)));
 	Hash = HashCombine(Hash, GetTypeHash(Opts.TaskId));
+	Hash = HashCombine(Hash, GetTypeHash(static_cast<int8>(Opts.TaskPriority)));
+	Hash = HashCombine(Hash, GetTypeHash(static_cast<int32>(Opts.ExtendedTaskPriority)));
 	
 	return Hash;
 }
diff --git a/Engine/Plugins/Experimental/AsyncMessageSystem/Source/AsyncMessageSystem/Public/AsyncMessageBindingOptions.h b/Engine/Plugins/Experimental/AsyncMessageSystem/Source/AsyncMessageSystem/Public/AsyncMessageBindingOptions.h
--- a/Engine/Plugins/Experimental/AsyncMessageSystem/Source/AsyncMessageSystem/Public/AsyncMessageBindingOptions.h
+++ b/Engine/Plugins/Experimental/AsyncMessageSystem/Source/AsyncMessageSystem/Public/AsyncMessageBindingOptions.h
@@ -3,15 +3,34 @@
 #pragma once
 
 #include "Engine/EngineBaseTypes.h"			// For ETickGround
-#include "Async/TaskGraphInterfaces.h"		// for ENamedThreads
+#include "Async/TaskGraphInterfaces.h"		// For ENamedThreads
+#include "Tasks/TaskPrivate.h"				// For ETaskPriority/EExtendedTaskPriority
 
 /**
  * Options used to specify when a binder would like to receive a message when it is broadcast.
+ *
+ * If you would like your listener to receive its messages within a specific tick group on the game thread,
+ * then you can use binding options with the Type set to "UseTickGroup", and specify the tick group you want. 
+ * 
+ * If you would like your listener to receive its messages within a specific named thread and/or pool of named threads 
+ * then you can specify the binding options with a Type of "UseTaskId", and specify the named thread on the TaskId.
+ *
+ * If you would like your listener to receive its messages from within a specifc set of UE::Tasks priority thread pools
+ * then you can specify the binding options with a Type of "UseTaskPriorities", and then set the task priority and extended
+ * priority. 
  */
 struct ASYNCMESSAGESYSTEM_API FAsyncMessageBindingOptions final
 {
 	FAsyncMessageBindingOptions() = default;
 	
+	FAsyncMessageBindingOptions(const ETickingGroup DesiredTickGroup);
+	
+	FAsyncMessageBindingOptions(const ENamedThreads::Type NamedThread);
+	
+	FAsyncMessageBindingOptions(
+		const UE::Tasks::ETaskPriority TaskPriority,
+		const UE::Tasks::EExtendedTaskPriority ExtendedTaskPriority);	
+	
 	/**
 	 * Determines which type of configuration this listener follows
 	 */
@@ -19,29 +38,53 @@ struct ASYNCMESSAGESYSTEM_API FAsyncMessageBindingOptions final
 	{
 		/**
 		 * If set, then this binding will receive a message on the main game thread during the specific tick group.
+		 *
+		 * The "Group" field will be used to identify which tick group to use.
 		 */
 		UseTickGroup = 0,
 
 		/**
 		 * If set, then this binding will receive a message on this TaskId's thread at some point
+		 * The "TaskId" field will be used to determine which Named Thread to you would like to use.
 		 */
 		UseTaskId = 1,
 
-		// TODO: UE::Tasks priority and extended priority groups
+		/**
+		 * If set, then this binding will receive a message on a thread determined by the UE::Tasks::ETaskPriority
+		 * and UE::Tasks::EExtendedTaskPriority.
+		 * 
+		 * The TaskPriority and ExtendedTaskPriority fields will be used to determine the UE::Tasks to launch
+		 * which wil process your bindings. 
+		 */
+		UseTaskPriorities = 2,
 	};
 
+	/**
+	 * The type of binding options which should be used to determine where a listener would like
+	 * to receive its messages. 
+	 */
 	EBindingType Type = EBindingType::UseTickGroup;
 
 	/**
 	 * Identifies the tick group which the binding object would like to receive a callback in
 	 */
 	ETickingGroup Group = ETickingGroup::TG_PostUpdateWork;
-		
+	
 	/**
 	 * Identifies the thread ID in which the binding object would like to receive a callback on
 	 */
 	int32 TaskId = ENamedThreads::Type::GameThread;
 
+	/**
+	 * Identifies the UE::Tasks task priority which this binding would like to receive a callback on
+	 */
+	UE::Tasks::ETaskPriority TaskPriority = UE::Tasks::ETaskPriority::Normal;
+
+	/**
+	* Identifies the UE::Tasks extended priority which this binding would like to receive a callback on
+	*/
+    UE::Tasks::EExtendedTaskPriority ExtendedTaskPriority = UE::Tasks::EExtendedTaskPriority::None;
+
 	bool operator==(const FAsyncMessageBindingOptions& Other) const;
 	bool operator!=(const FAsyncMessageBindingOptions& Other) const;
 	ASYNCMESSAGESYSTEM_API friend uint32 GetTypeHash(const FAsyncMessageBindingOptions& Opts);
