Fix a bug with moving payload data when queueing an event, which was causing the payload to be null when you had multiple different bindings to the same event. We can resolve this by simply copying the payload data to the queue instead of using the incorrect MoveTemp call.

#rnx

diff --git a/Engine/Plugins/Experimental/AsyncMessageSystem/Source/AsyncMessageSystem/Private/AsyncMessageStore.cpp b/Engine/Plugins/Experimental/AsyncMessageSystem/Source/AsyncMessageSystem/Private/AsyncMessageStore.cpp
--- a/Engine/Plugins/Experimental/AsyncMessageSystem/Source/AsyncMessageSystem/Private/AsyncMessageStore.cpp
+++ b/Engine/Plugins/Experimental/AsyncMessageSystem/Source/AsyncMessageSystem/Private/AsyncMessageStore.cpp
@@ -22,7 +22,7 @@ uint32 FAsyncMessageStore::EnqueueMessage(FAsyncMessage&& MessageToQueue, TSet<F
 			// Only add this message to a queue if we know that it has a listener for this particular message.
 			if (Pair.Value.MessagesWithTheseBindings.Contains(CurrentMessageId) && !BindingTypesAddedTo.Contains(Pair.Key))
 			{
-				Pair.Value.MessageQueue->Enqueue(MoveTemp(MessageToQueue));
+				Pair.Value.MessageQueue->Enqueue(MessageToQueue);
 				++NumQueuesAddedTo;
 			
 				BindingTypesAddedTo.Add(Pair.Key);
@@ -72,7 +72,7 @@ bool FAsyncMessageStore::RemoveMessageFromBinding(const FAsyncMessageId& Message
 		return false;
 	}
 
-	// Decremen the count of how many messages of this ID are on this binding option
+	// Decrement the count of how many messages of this ID are on this binding option
 	if (uint32* Count = QueueData->MessagesWithTheseBindings.Find(MessageId))
 	{
 		(*Count)--;
