Enable unsafe type cast warnings in the Async Message system plugin and fix one up.

Also add a quick line which you can uncomment to turn off optimization of the module, making debugging much easier.

Some misc. include file organization to make it alphabetical, fix some typos in comments.

#rnx

diff --git a/Engine/Plugins/Experimental/AsyncMessageSystem/Source/AsyncMessageSystem/AsyncMessageSystem.build.cs b/Engine/Plugins/Experimental/AsyncMessageSystem/Source/AsyncMessageSystem/AsyncMessageSystem.build.cs
--- a/Engine/Plugins/Experimental/AsyncMessageSystem/Source/AsyncMessageSystem/AsyncMessageSystem.build.cs
+++ b/Engine/Plugins/Experimental/AsyncMessageSystem/Source/AsyncMessageSystem/AsyncMessageSystem.build.cs
@@ -7,6 +7,12 @@ namespace UnrealBuildTool.Rules
 		public AsyncMessageSystem(ReadOnlyTargetRules Target) : base(Target)
 		{
 			PCHUsage = ModuleRules.PCHUsageMode.UseExplicitOrSharedPCHs;
+
+			// Enable truncation warnings in this plugin
+			CppCompileWarningSettings.UnsafeTypeCastWarningLevel = WarningLevel.Error;
+			
+			// Uncomment this line to make for easier debugging
+			//OptimizeCode = CodeOptimization.Never;
 			
 			PublicDependencyModuleNames.AddRange(
 				new string[] {
diff --git a/Engine/Plugins/Experimental/AsyncMessageSystem/Source/AsyncMessageSystem/Private/AsyncGameplayMessageSystem.cpp b/Engine/Plugins/Experimental/AsyncMessageSystem/Source/AsyncMessageSystem/Private/AsyncGameplayMessageSystem.cpp
--- a/Engine/Plugins/Experimental/AsyncMessageSystem/Source/AsyncMessageSystem/Private/AsyncGameplayMessageSystem.cpp
+++ b/Engine/Plugins/Experimental/AsyncMessageSystem/Source/AsyncMessageSystem/Private/AsyncGameplayMessageSystem.cpp
@@ -1,8 +1,9 @@
 // Copyright Epic Games, Inc. All Rights Reserved.
 
 #include "AsyncGameplayMessageSystem.h"
-#include "AsyncMessageSystemLogs.h"
+
 #include "Async/Async.h"
+#include "AsyncMessageSystemLogs.h"
 #include "Engine/Level.h"
 #include "Engine/World.h"
 #include "Tasks/TaskPrivate.h"
@@ -200,7 +201,7 @@ void FAsyncGameplayMessageSystem::StartAsyncProcessForBinding(const FAsyncMessag
 	// });
 	//
 	// We utilize the UE::Tasks system here and translate the ENamedThread because the UE::Tasks system
-	// should have less overhead and better schdueling capabilities in most scenarios
+	// should have less overhead and better scheduling capabilities in most scenarios
 }
 
 void FAsyncGameplayMessageSystem::ExecuteTick(
diff --git a/Engine/Plugins/Experimental/AsyncMessageSystem/Source/AsyncMessageSystem/Private/AsyncMessage.cpp b/Engine/Plugins/Experimental/AsyncMessageSystem/Source/AsyncMessageSystem/Private/AsyncMessage.cpp
--- a/Engine/Plugins/Experimental/AsyncMessageSystem/Source/AsyncMessageSystem/Private/AsyncMessage.cpp
+++ b/Engine/Plugins/Experimental/AsyncMessageSystem/Source/AsyncMessageSystem/Private/AsyncMessage.cpp
@@ -19,9 +19,11 @@ FAsyncMessage::FAsyncMessage(
 	, PayloadCopy(InPayloadData)	// <-- This is making a copy of the payload data!!! $$$
 									// It calls the explicit FInstancedStruct(const FConstStructView InOther) constructor
 									// which is going to copy the UScriptStruct.
+									// This is the same as : 
+									//		PayloadCopy.InitializeAs(InPayloadData.GetScriptStruct(), InPayloadData.GetMemory());
 									// TODO: Use some custom linear allocator here on the message system to make the copying cheaper
 {
-	
+
 }
 
 double FAsyncMessage::GetQueueTimestamp() const
@@ -59,7 +61,12 @@ FAsyncMessageId FAsyncMessage::GetMessageSourceId() const
 	return MessageSourceId;
 }
 
-FStructView FAsyncMessage::GetPayloadView() const
+FStructView FAsyncMessage::GetPayloadView()
+{
+	return FStructView(PayloadCopy);
+}
+
+FConstStructView FAsyncMessage::GetPayloadView() const
 {
-	return PayloadCopy.GetScriptStruct();
+	return FConstStructView(PayloadCopy);
 }
diff --git a/Engine/Plugins/Experimental/AsyncMessageSystem/Source/AsyncMessageSystem/Private/AsyncMessageSystemBase.cpp b/Engine/Plugins/Experimental/AsyncMessageSystem/Source/AsyncMessageSystem/Private/AsyncMessageSystemBase.cpp
--- a/Engine/Plugins/Experimental/AsyncMessageSystem/Source/AsyncMessageSystem/Private/AsyncMessageSystemBase.cpp
+++ b/Engine/Plugins/Experimental/AsyncMessageSystem/Source/AsyncMessageSystem/Private/AsyncMessageSystemBase.cpp
@@ -1,6 +1,7 @@
 // Copyright Epic Games, Inc. All Rights Reserved.
 
 #include "AsyncMessageSystemBase.h"
+
 #include "AsyncMessageSystemLogs.h"
 #include "CoreGlobals.h"
 #include "Misc/App.h"
@@ -162,7 +163,7 @@ bool FAsyncMessageSystemBase::QueueMessageForBroadcast_Impl(const FAsyncMessageI
 	// Note: We have FStructView here, which should give us some reflected data about the type of
 	// payload we are receiving.
 
-	const uint64 QueueTime = FApp::GetCurrentTime();
+	const double QueueTime = FApp::GetCurrentTime();
 	const uint64 QueuedFrame = GFrameCounter;
 	const uint32 ThreadQueuedFrom = FPlatformTLS::GetCurrentThreadId();
 	const uint32 MessageSequence = NextMessageSequence++;
diff --git a/Engine/Plugins/Experimental/AsyncMessageSystem/Source/AsyncMessageSystem/Private/AsyncMessageWorldSubsystem.cpp b/Engine/Plugins/Experimental/AsyncMessageSystem/Source/AsyncMessageSystem/Private/AsyncMessageWorldSubsystem.cpp
--- a/Engine/Plugins/Experimental/AsyncMessageSystem/Source/AsyncMessageSystem/Private/AsyncMessageWorldSubsystem.cpp
+++ b/Engine/Plugins/Experimental/AsyncMessageSystem/Source/AsyncMessageSystem/Private/AsyncMessageWorldSubsystem.cpp
@@ -1,8 +1,9 @@
 // Copyright Epic Games, Inc. All Rights Reserved.
 
 #include "AsyncMessageWorldSubsystem.h"
-#include "AsyncMessageSystemLogs.h"
+
 #include "AsyncGameplayMessageSystem.h"
+#include "AsyncMessageSystemLogs.h"
 
 void UAsyncMessageWorldSubsystem::Initialize(FSubsystemCollectionBase& Collection)
 {
@@ -22,7 +23,7 @@ void UAsyncMessageWorldSubsystem::InitMessageSystem()
 {
 	check(!MessageSystem.IsValid());
 	
-	UE_LOG(LogAsyncMessageSystem, Log, TEXT("[%hs] Init world message system for world %s..."), __func__, *GetNameSafe(GetWorld()));
+	UE_LOG(LogAsyncMessageSystem, Verbose, TEXT("[%hs] Init world message system for world '%s'"), __func__, *GetNameSafe(GetWorld()));
 
 	MessageSystem = FAsyncMessageSystemBase::CreateMessageSystem<FAsyncGameplayMessageSystem>(GetWorld());
 }
@@ -31,7 +32,7 @@ void UAsyncMessageWorldSubsystem::ShutdownMessageSystem()
 {
 	if (ensure(MessageSystem.IsValid()))
 	{
-		UE_LOG(LogAsyncMessageSystem, Log, TEXT("[%hs] Shutdown world message system for world %s..."), __func__, *GetNameSafe(GetWorld()));
+		UE_LOG(LogAsyncMessageSystem, Verbose, TEXT("[%hs] Shutdown world message system for world '%s'"), __func__, *GetNameSafe(GetWorld()));
 		
 		MessageSystem->Shutdown();
 		
diff --git a/Engine/Plugins/Experimental/AsyncMessageSystem/Source/AsyncMessageSystem/Public/AsyncGameplayMessageSystem.h b/Engine/Plugins/Experimental/AsyncMessageSystem/Source/AsyncMessageSystem/Public/AsyncGameplayMessageSystem.h
--- a/Engine/Plugins/Experimental/AsyncMessageSystem/Source/AsyncMessageSystem/Public/AsyncGameplayMessageSystem.h
+++ b/Engine/Plugins/Experimental/AsyncMessageSystem/Source/AsyncMessageSystem/Public/AsyncGameplayMessageSystem.h
@@ -2,14 +2,14 @@
 
 #pragma once
 
+#include "Async/TaskGraphFwd.h"			// FGraphEventRef
 #include "AsyncMessageSystemBase.h"
 #include "Engine/EngineBaseTypes.h"		// For ELevelTick
-#include "Async/TaskGraphFwd.h"			// FGraphEventRef
 
 class UWorld;
 
 /**
- * Implementation of an async message system which scheudles the processing of messages based on tick groups
+ * Implementation of an async message system which schedules the processing of messages based on tick groups
  * and named thread async tasks, making it integrate nicely with common gameplay frameworks in Unreal.
  */
 class ASYNCMESSAGESYSTEM_API FAsyncGameplayMessageSystem final :
diff --git a/Engine/Plugins/Experimental/AsyncMessageSystem/Source/AsyncMessageSystem/Public/AsyncMessage.h b/Engine/Plugins/Experimental/AsyncMessageSystem/Source/AsyncMessageSystem/Public/AsyncMessage.h
--- a/Engine/Plugins/Experimental/AsyncMessageSystem/Source/AsyncMessageSystem/Public/AsyncMessage.h
+++ b/Engine/Plugins/Experimental/AsyncMessageSystem/Source/AsyncMessageSystem/Public/AsyncMessage.h
@@ -2,8 +2,8 @@
 
 #pragma once
 
-#include "HAL/Platform.h"	// For uint64 type decl
 #include "AsyncMessageId.h"
+#include "HAL/Platform.h"	// For uint64 type declaration
 #include "StructUtils/StructView.h"
 
 /**
@@ -11,7 +11,7 @@
  * interested in receiving the data from.
  *
  * Async Message Messages contain some data about when, where, and how the message was queued
- * which make it possible to synconize game state across multiple threads with messages. 
+ * which make it possible to synchronize game state across multiple threads with messages. 
  */
 struct ASYNCMESSAGESYSTEM_API FAsyncMessage final
 {
@@ -29,7 +29,7 @@ public:
 		FConstStructView InPayloadData);
 
 	/**
-	 * Returns this message’s unique ID which listeners use to bind to it. 
+	 * Returns this message's unique ID which listeners use to bind to it. 
 	 */
 	FAsyncMessageId GetMessageId() const;
 
@@ -68,9 +68,14 @@ public:
 	uint32 GetSequenceId() const;
 
 	/**
-	 * Get the payload data associated with this async message.
+	 * Get the mutable payload data associated with this async message.
 	 */
-	FStructView GetPayloadView() const;
+	FStructView GetPayloadView();
+
+	/**
+	 * Get the const payload data associated with this async message.
+	 */
+	FConstStructView GetPayloadView() const;
 
 	/**
 	 * Returns the typed data pointer of this message's payload. Nullptr if the payload is invalid.
diff --git a/Engine/Plugins/Experimental/AsyncMessageSystem/Source/AsyncMessageSystem/Public/AsyncMessageBindingOptions.h b/Engine/Plugins/Experimental/AsyncMessageSystem/Source/AsyncMessageSystem/Public/AsyncMessageBindingOptions.h
--- a/Engine/Plugins/Experimental/AsyncMessageSystem/Source/AsyncMessageSystem/Public/AsyncMessageBindingOptions.h
+++ b/Engine/Plugins/Experimental/AsyncMessageSystem/Source/AsyncMessageSystem/Public/AsyncMessageBindingOptions.h
@@ -2,8 +2,8 @@
 
 #pragma once
 
-#include "Engine/EngineBaseTypes.h"			// For ETickGround
 #include "Async/TaskGraphInterfaces.h"		// For ENamedThreads
+#include "Engine/EngineBaseTypes.h"			// For ETickGround
 #include "Tasks/TaskPrivate.h"				// For ETaskPriority/EExtendedTaskPriority
 
 /**
@@ -54,7 +54,7 @@ struct ASYNCMESSAGESYSTEM_API FAsyncMessageBindingOptions final
 		 * and UE::Tasks::EExtendedTaskPriority.
 		 * 
 		 * The TaskPriority and ExtendedTaskPriority fields will be used to determine the UE::Tasks to launch
-		 * which wil process your bindings. 
+		 * which will process your bindings. 
 		 */
 		UseTaskPriorities = 2,
 	};
diff --git a/Engine/Plugins/Experimental/AsyncMessageSystem/Source/AsyncMessageSystem/Public/AsyncMessageHandle.h b/Engine/Plugins/Experimental/AsyncMessageSystem/Source/AsyncMessageSystem/Public/AsyncMessageHandle.h
--- a/Engine/Plugins/Experimental/AsyncMessageSystem/Source/AsyncMessageSystem/Public/AsyncMessageHandle.h
+++ b/Engine/Plugins/Experimental/AsyncMessageSystem/Source/AsyncMessageSystem/Public/AsyncMessageHandle.h
@@ -2,9 +2,9 @@
 
 #pragma once
 
-#include "Containers/UnrealString.h"	// For FString
-#include "HAL/Platform.h"	// For uint32 type decl
 #include "AsyncMessageId.h"
+#include "Containers/UnrealString.h"	// For FString
+#include "HAL/Platform.h"				// For uint32 type declaration
 
 /**
  * Handle used to identify a bound listener to an Async Message.
@@ -65,7 +65,7 @@ private:
 	/**
 	 * Internal index for message handles which we consider to be invalid.
 	 */
-	static constexpr uint32 InvalidHandleIndex = 0;
+	static constexpr uint32 InvalidHandleIndex = 0u;
 	
 	/**
 	 * The value of this internal handle.
diff --git a/Engine/Plugins/Experimental/AsyncMessageSystem/Source/AsyncMessageSystem/Public/AsyncMessageId.h b/Engine/Plugins/Experimental/AsyncMessageSystem/Source/AsyncMessageSystem/Public/AsyncMessageId.h
--- a/Engine/Plugins/Experimental/AsyncMessageSystem/Source/AsyncMessageSystem/Public/AsyncMessageId.h
+++ b/Engine/Plugins/Experimental/AsyncMessageSystem/Source/AsyncMessageSystem/Public/AsyncMessageId.h
@@ -2,9 +2,9 @@
 
 #pragma once
 
-#include "UObject/NameTypes.h"			// For FName
 #include "Containers/UnrealString.h"	// For FString
 #include "GameplayTagContainer.h"		// For FGameplayTag
+#include "UObject/NameTypes.h"			// For FName
 
 #include "AsyncMessageId.generated.h"
 
@@ -27,14 +27,14 @@ public:
 	/**
 	 * Constructs a MessageId from the given FName.
 	 *
-	 * Note: This will internally be requesting a gameplay tag representation of this Fname, and will ensure if it is not found.
+	 * Note: This will internally be requesting a gameplay tag representation of this FName, and will ensure if it is not found.
 	 * 
 	 * @param MessageName The FName representation of this MessageId
 	 */
 	FAsyncMessageId(const FName MessageName);
 	
 	/**
-	 * Constructs a FAsyncMessageId from the FName assocaited with the given Gameplay Tag. 
+	 * Constructs a FAsyncMessageId from the FName associated with the given Gameplay Tag. 
 	 * 
 	 * @param MessageTag The gameplay tag to construct a Message Id from
 	 */
diff --git a/Engine/Plugins/Experimental/AsyncMessageSystem/Source/AsyncMessageSystem/Public/AsyncMessageWorldSubsystem.h b/Engine/Plugins/Experimental/AsyncMessageSystem/Source/AsyncMessageSystem/Public/AsyncMessageWorldSubsystem.h
--- a/Engine/Plugins/Experimental/AsyncMessageSystem/Source/AsyncMessageSystem/Public/AsyncMessageWorldSubsystem.h
+++ b/Engine/Plugins/Experimental/AsyncMessageSystem/Source/AsyncMessageSystem/Public/AsyncMessageWorldSubsystem.h
@@ -2,8 +2,8 @@
 
 #pragma once
 
-#include "Subsystems/WorldSubsystem.h"
 #include "Engine/World.h"
+#include "Subsystems/WorldSubsystem.h"
 
 #include "AsyncMessageWorldSubsystem.generated.h"
 
@@ -15,10 +15,8 @@ class FAsyncMessageSystemBase;
  *
  * An example of using the message system like this is:
  *
- * if (TSharedPtr<IAsyncMessageSystem> Sys = UAsyncMessageWorldSubsystem::GetSharedMessageSystem(GetWorld()))
- * {
- *		Sys->QueueMessageForBroadcast(MessageId, YourPayloadData);
- * }
+ *		TSharedPtr<FAsyncMessageSystemBase> Sys = UAsyncMessageWorldSubsystem::GetSharedMessageSystem(GetWorld());
+ *		Sys->QueueMessageForBroadcast(MessageToQueue, MessagePayloadToQueue);
  */
 UCLASS()
 class ASYNCMESSAGESYSTEM_API UAsyncMessageWorldSubsystem :
@@ -28,37 +26,39 @@ class ASYNCMESSAGESYSTEM_API UAsyncMessageWorldSubsystem :
 GENERATED_BODY()
 
 public:
-
-	template<class TMessageSystemType = FAsyncMessageSystemBase>
-	static TSharedPtr<TMessageSystemType> GetSharedMessageSystem(const UWorld* InWorld)
-	{
-		static_assert(std::is_base_of<FAsyncMessageSystemBase, TMessageSystemType>::value, "TMessageSystemType not derived from FAsyncMessageSystemBase");
-
-		if (InWorld)
-		{
-			if (const UAsyncMessageWorldSubsystem* Sys = InWorld->GetSubsystem<UAsyncMessageWorldSubsystem>())
-			{
-				return Sys->GetSharedMessageSystem<TMessageSystemType>();	
-			}
-		}
 	
-		return nullptr;
-	}
-
+	/**
+	 * Static helper function to get the message system for a world.
+	 * 
+	 * @tparam TMessageSystemType	The type of message system to create, must be a subclass of FAsyncMessageSystemBase
+	 * @param InWorld				The world to get the message system of
+	 * @return						Shared pointer to the given world's message system. Null if the given world is invalid.
+	 */
 	template<class TMessageSystemType = FAsyncMessageSystemBase>
-	TSharedPtr<TMessageSystemType> GetSharedMessageSystem() const
-	{
-		static_assert(std::is_base_of<FAsyncMessageSystemBase, TMessageSystemType>::value, "TMessageSystemType not derived from FAsyncMessageSystemBase");
+	static TSharedPtr<TMessageSystemType> GetSharedMessageSystem(const UWorld* InWorld);
 
-		return StaticCastSharedPtr<TMessageSystemType>(MessageSystem);
-	}
+	/**
+	 * Returns the message system for this world subsystem.
+	 * 
+	 * @tparam TMessageSystemType	The type of message system to create, must be a subclass of FAsyncMessageSystemBase
+	 * @return						Shared pointer to this world subsystem's current message system
+	 */
+	template<class TMessageSystemType = FAsyncMessageSystemBase>
+	TSharedPtr<TMessageSystemType> GetSharedMessageSystem() const;
 	
 protected:
 
 	virtual void Initialize(FSubsystemCollectionBase& Collection) override;
 	virtual void Deinitialize() override;
 
+	/**
+	 * Creates the MessageSystem for this world subsystem. 
+	 */
 	void InitMessageSystem();
+
+	/**
+	* Shutdown and cleans up this worlds message system. 
+	*/
 	void ShutdownMessageSystem();
 
 	/**
@@ -67,3 +67,31 @@ protected:
 	 */
 	TSharedPtr<FAsyncMessageSystemBase, ESPMode::ThreadSafe> MessageSystem = nullptr;
 };
+
+//////////////////////////////////////////////////////////////
+// Template Definitions
+//////////////////////////////////////////////////////////////
+
+template <class TMessageSystemType>
+TSharedPtr<TMessageSystemType> UAsyncMessageWorldSubsystem::GetSharedMessageSystem(const UWorld* InWorld)
+{
+	static_assert(std::is_base_of<FAsyncMessageSystemBase, TMessageSystemType>::value, "TMessageSystemType not derived from FAsyncMessageSystemBase");
+
+	if (InWorld)
+	{
+		if (const UAsyncMessageWorldSubsystem* Sys = InWorld->GetSubsystem<UAsyncMessageWorldSubsystem>())
+		{
+			return Sys->GetSharedMessageSystem<TMessageSystemType>();	
+		}
+	}
+
+	return nullptr;
+}
+
+template <class TMessageSystemType>
+TSharedPtr<TMessageSystemType> UAsyncMessageWorldSubsystem::GetSharedMessageSystem() const
+{
+	static_assert(std::is_base_of<FAsyncMessageSystemBase, TMessageSystemType>::value, "TMessageSystemType not derived from FAsyncMessageSystemBase");
+
+	return StaticCastSharedPtr<TMessageSystemType>(MessageSystem);
+}
