Misc. Cleanup of async message system unit tests. Make the test system within the automation tests #define, and mark some heavier tests with the StressTest flag to save on boot up time.

Add a new test to ensure that you can read a valid payload when you have mulitple binding optoins to a single message.

#rnx

diff --git a/Engine/Plugins/Experimental/AsyncMessageSystem/Source/AsyncMessageSystem/Private/Tests/AsyncMessageSystemTests.cpp b/Engine/Plugins/Experimental/AsyncMessageSystem/Source/AsyncMessageSystem/Private/Tests/AsyncMessageSystemTests.cpp
--- a/Engine/Plugins/Experimental/AsyncMessageSystem/Source/AsyncMessageSystem/Private/Tests/AsyncMessageSystemTests.cpp
+++ b/Engine/Plugins/Experimental/AsyncMessageSystem/Source/AsyncMessageSystem/Private/Tests/AsyncMessageSystemTests.cpp
@@ -1,10 +1,12 @@
 // Copyright Epic Games, Inc. All Rights Reserved.
 
 #include "AsyncMessageSystemTests.h"
-#include "Misc/AutomationTest.h"
+
 #include "AsyncMessageHandle.h"
-#include "Tasks/Task.h"
+#include "AsyncMessageSystemBase.h"
+#include "Misc/AutomationTest.h"
 #include "NativeGameplayTags.h"
+#include "Tasks/Task.h"
 
 UE_DEFINE_GAMEPLAY_TAG_COMMENT(InternalTestTag_Invalid, "AsyncMessages.Internal.test.invalid", "A test gameplay tag utilized in the async message system unit tests")
 UE_DEFINE_GAMEPLAY_TAG_COMMENT(InternalTestTag_A, "AsyncMessages.Internal.test.a", "A test gameplay tag utilized in the async message system unit tests")
@@ -18,6 +20,97 @@ namespace UE::AsyncMessageSystem
 {
 	constexpr EAutomationTestFlags QuickTestFlags = EAutomationTestFlags_ApplicationContextMask | EAutomationTestFlags::SmokeFilter;
 	constexpr EAutomationTestFlags StressTestFlags = EAutomationTestFlags_ApplicationContextMask | EAutomationTestFlags::StressFilter;
+	
+	// Test message system interface impl
+	class FTestMessageSystem : public FAsyncMessageSystemBase
+	{
+	protected:
+		virtual void Startup_Impl() override
+		{
+			bHasStarted = true;
+		}
+
+		virtual void Shutdown_Impl() override
+		{
+			bHasShutdown = true;
+		}
+
+		virtual void PostQueueMessage(const FAsyncMessageId MessageId, const TSet<FAsyncMessageBindingOptions>& OptionsBoundTo) override
+		{
+			// Dont do anything by default here in this simple test message system
+		}
+		
+		bool bHasStarted = false;
+		bool bHasShutdown = false;	
+
+	public:
+		const bool HasStarted() const { return bHasStarted; }
+		const bool HasShutdown() const { return bHasShutdown && bIsShuttingDown; }
+		
+		/**
+		 * Returns true if the handle currently has any listeners bound to it in the message map
+		 */
+		bool IsHandleBound(const FAsyncMessageHandle& Handle) const
+		{
+			FScopeLock ListenerLock(&MessageListenerMapCS);
+
+			if (const FAsyncMessageBoundData* FoundBindingData = BoundMessageListenerMap.Find(Handle.GetBoundMessageId()))
+			{
+				for (const TPair<FAsyncMessageBindingOptions, TArray<FAsyncMessageIndividualListener<FMessageCallbackFunc>>>& ListenerPair : FoundBindingData->ListenerMap)
+				{
+					const bool bHasHandle = ListenerPair.Value.ContainsByPredicate([Handle](const FAsyncMessageIndividualListener<FMessageCallbackFunc>& Listener)
+					{
+						return Listener.Handle == Handle;	
+					});
+		
+					if (bHasHandle)
+					{
+						return true;
+					}
+				}	
+			}
+			
+			return false;	
+		}
+
+		uint32 GetNumberOfListeners() const
+		{
+			uint32 Res = 0;
+			// Check any message listeners...
+			{
+				FScopeLock ListenerLock(&MessageListenerMapCS);
+				for (const TPair<FAsyncMessageId, FAsyncMessageBoundData>& MessagePair : BoundMessageListenerMap)
+				{
+					const FAsyncMessageBoundData& BoundData = MessagePair.Value;
+					for (const TPair<FAsyncMessageBindingOptions, TArray<FAsyncMessageIndividualListener<FMessageCallbackFunc>>>& ListenerPair : BoundData.ListenerMap)
+					{
+						Res += ListenerPair.Value.Num();
+					}
+				}	
+			}
+
+			return Res;
+		}
+		
+		void RunOnce(const TArray<FAsyncMessageBindingOptions>& OptsToProcess = { FAsyncMessageBindingOptions{} })
+		{
+			for (const FAsyncMessageBindingOptions& Options : OptsToProcess)
+			{
+				ProcessMessageQueueForBinding(FAsyncMessageBindingOptions{});
+			}			
+		}
+
+		void Test_Shutdown()
+		{
+			Shutdown();
+		}
+
+		FAsyncMessageHandle GenerateTestHandle()
+		{
+			return GenerateNextValidMessageHandle({InternalTestTag_Invalid});
+		}
+	};
+	
 
 	IMPLEMENT_SIMPLE_AUTOMATION_TEST(TMessageHandleDefaults, "AsyncMessagePassing.MessageHandles", QuickTestFlags)
 	bool TMessageHandleDefaults::RunTest(const FString& InParameters)
@@ -203,7 +296,7 @@ namespace UE::AsyncMessageSystem
 				// We should be able to add a test binding in response to 
 				if (const FNested_Payload* Data = Message.GetPayloadData<const FNested_Payload>())
 				{
-					if (TSharedPtr<FTestMessageSystem> Sys = Data->MessageSystem.Pin())
+					if (TSharedPtr<FAsyncMessageSystemBase> Sys = Data->MessageSystem.Pin())
 					{
 						// Bind a nested listener here
 						const FAsyncMessageHandle Handle = Sys->BindListener(TestMessageId_A, [](const FAsyncMessage& Message)
@@ -282,6 +375,55 @@ namespace UE::AsyncMessageSystem
 		return true;
 	}	
 
+	// Test to make sure that a single message can be queued and broadcast to multiple different kinds of binding options.
+	IMPLEMENT_SIMPLE_AUTOMATION_TEST(TMessageSystemBindingMultipleOptions, "AsyncMessagePassing.Binding.MultipleOptions", QuickTestFlags)
+	bool TMessageSystemBindingMultipleOptions::RunTest(const FString& InParameters)
+	{
+		TSharedPtr<FTestMessageSystem, ESPMode::ThreadSafe> TestSystem = FAsyncMessageSystemBase::CreateMessageSystem<FTestMessageSystem>();
+		TestTrue(TEXT("MessageSystem startup has been called"), TestSystem->HasStarted());
+
+		auto TestLambda = [&](const FAsyncMessage& Message)
+		{
+			const FTest_Payload_A* Data = Message.GetPayloadData<const FTest_Payload_A>();
+			TestNotNull(TEXT("Const Payload data is valid "), Data);
+
+			FConstStructView DataView = Message.GetPayloadView();
+			TestTrue(TEXT("Const Payload view is valid"), DataView.IsValid());			
+		};
+
+		// Bind some test listeners which will ensure that the payload data is valid
+		const TArray<FAsyncMessageBindingOptions> BindingsToUse = 
+		{
+			FAsyncMessageBindingOptions(ETickingGroup::TG_PrePhysics),
+			FAsyncMessageBindingOptions(ETickingGroup::TG_DuringPhysics),
+			FAsyncMessageBindingOptions(ETickingGroup::TG_PostPhysics),
+			FAsyncMessageBindingOptions(ETickingGroup::TG_PostUpdateWork),
+			FAsyncMessageBindingOptions(ENamedThreads::HighTaskPriority),
+			FAsyncMessageBindingOptions(ENamedThreads::GameThread),
+			FAsyncMessageBindingOptions(ENamedThreads::RHIThread),
+			FAsyncMessageBindingOptions(UE::Tasks::ETaskPriority::Default, UE::Tasks::EExtendedTaskPriority::Inline),
+			FAsyncMessageBindingOptions(UE::Tasks::ETaskPriority::ForegroundCount, UE::Tasks::EExtendedTaskPriority::TaskEvent),
+			FAsyncMessageBindingOptions(UE::Tasks::ETaskPriority::BackgroundNormal, UE::Tasks::EExtendedTaskPriority::GameThreadHiPri),
+		};
+		
+		for (const FAsyncMessageBindingOptions& Opts : BindingsToUse)
+		{
+			// Bind listeners to multiple different types
+			TestSystem->BindListener(TestMessageId_A, TestLambda, Opts);
+		}
+		
+		// Queue a test message for broadcasting
+		FInstancedStruct PayloadDataInstance = FInstancedStruct::Make<FTest_Payload_A>(FTest_Payload_A{});
+		TestSystem->QueueMessageForBroadcast(TestMessageId_A, PayloadDataInstance);
+
+		// Actually run the system, which should process all the messages in the queue.
+		TestSystem->RunOnce(BindingsToUse);
+
+		TestSystem->Shutdown();
+		TestSystem.Reset();
+		return true;
+	}
+
 	// Test binding a simple C++ lambda 
 	IMPLEMENT_SIMPLE_AUTOMATION_TEST(TMessageSystemBindingLambdaTest, "AsyncMessagePassing.Binding.Lambda", QuickTestFlags)
 	bool TMessageSystemBindingLambdaTest::RunTest(const FString& InParameters)
@@ -473,7 +615,7 @@ namespace UE::AsyncMessageSystem
 	// A test which will spin up several tasks on different threads to attempt to bind listeners to messages
 	// on each thread. This will test that we can have different listeners binding on different threads
 	// for the same message system
-	IMPLEMENT_SIMPLE_AUTOMATION_TEST(TMessageSystemMultiThreadedBindListeners, "AsyncMessagePassing.MultiThreaded.BindListeners", QuickTestFlags)
+	IMPLEMENT_SIMPLE_AUTOMATION_TEST(TMessageSystemMultiThreadedBindListeners, "AsyncMessagePassing.MultiThreaded.BindListeners", StressTestFlags)
 	bool TMessageSystemMultiThreadedBindListeners::RunTest(const FString& InParameters)
 	{
 		TSharedPtr<FTestMessageSystem, ESPMode::ThreadSafe> TestSystem = FAsyncMessageSystemBase::CreateMessageSystem<FTestMessageSystem>();
@@ -564,7 +706,7 @@ namespace UE::AsyncMessageSystem
 	}
 
 	// Test to queue messages from multiple threads on a single message system
-	IMPLEMENT_SIMPLE_AUTOMATION_TEST(TMessageSystemMultiThreadedSequenceIdTest, "AsyncMessagePassing.MultiThreaded.UniqueSequenceId", QuickTestFlags)
+	IMPLEMENT_SIMPLE_AUTOMATION_TEST(TMessageSystemMultiThreadedSequenceIdTest, "AsyncMessagePassing.MultiThreaded.UniqueSequenceId", StressTestFlags)
 	bool TMessageSystemMultiThreadedSequenceIdTest::RunTest(const FString& InParameters)
 	{
 		TSharedPtr<FTestMessageSystem, ESPMode::ThreadSafe> TestSystem = FAsyncMessageSystemBase::CreateMessageSystem<FTestMessageSystem>();
@@ -629,53 +771,7 @@ namespace UE::AsyncMessageSystem
 
 #endif	// WITH_DEV_AUTOMATION_TESTS
 
-bool UE::AsyncMessageSystem::FTestMessageSystem::IsHandleBound(const FAsyncMessageHandle& Handle) const
-{
-	FScopeLock ListenerLock(&MessageListenerMapCS);
-	for (const TPair<FAsyncMessageId, FAsyncMessageBoundData>& MessagePair : BoundMessageListenerMap)
-	{
-		const FAsyncMessageBoundData& BoundData = MessagePair.Value;
-		for (const TPair<FAsyncMessageBindingOptions, TArray<FAsyncMessageIndividualListener<FMessageCallbackFunc>>>& ListenerPair : BoundData.ListenerMap)
-		{
-			const bool bHasHandle = ListenerPair.Value.ContainsByPredicate([Handle](const FAsyncMessageIndividualListener<FMessageCallbackFunc>& Listener)
-			{
-				return Listener.Handle == Handle;	
-			});
-			
-			if (bHasHandle)
-			{
-				return true;
-			}
-		}
-	}	
-	
-	return false;	
-}
-
-uint32 UE::AsyncMessageSystem::FTestMessageSystem::GetNumberOfListeners() const
-{
-	uint32 Res = 0;
-	// Check any message listeners...
-	{
-		FScopeLock ListenerLock(&MessageListenerMapCS);
-		for (const TPair<FAsyncMessageId, FAsyncMessageBoundData>& MessagePair : BoundMessageListenerMap)
-		{
-			const FAsyncMessageBoundData& BoundData = MessagePair.Value;
-			for (const TPair<FAsyncMessageBindingOptions, TArray<FAsyncMessageIndividualListener<FMessageCallbackFunc>>>& ListenerPair : BoundData.ListenerMap)
-			{
-				Res += ListenerPair.Value.Num();
-			}
-		}	
-	}
-
-	return Res;
-}
-
-FAsyncMessageHandle UE::AsyncMessageSystem::FTestMessageSystem::GenerateTestHandle()
-{
-	return GenerateNextValidMessageHandle({InternalTestTag_Invalid});
-}
-
+// A test UFunction implmentation 
 void UTestAsyncObject::CallbackFunction(const FAsyncMessage& Message)
 {
 	FConstStructView DataView = Message.GetPayloadView();
diff --git a/Engine/Plugins/Experimental/AsyncMessageSystem/Source/AsyncMessageSystem/Private/Tests/AsyncMessageSystemTests.h b/Engine/Plugins/Experimental/AsyncMessageSystem/Source/AsyncMessageSystem/Private/Tests/AsyncMessageSystemTests.h
--- a/Engine/Plugins/Experimental/AsyncMessageSystem/Source/AsyncMessageSystem/Private/Tests/AsyncMessageSystemTests.h
+++ b/Engine/Plugins/Experimental/AsyncMessageSystem/Source/AsyncMessageSystem/Private/Tests/AsyncMessageSystemTests.h
@@ -2,12 +2,22 @@
 
 #pragma once
 
-#include "CoreMinimal.h"
-#include "UObject/Object.h"
-#include "AsyncGameplayMessageSystem.h"
+#include "Containers/UnrealString.h"	// For FString
+#include "HAL/Platform.h"				// For int32 typedef
+#include "Templates/SharedPointer.h"	// For TWeakPtr
+#include "UObject/Object.h"				// For UObject
+
 #include "AsyncMessageSystemTests.generated.h"
 
-USTRUCT()
+struct FAsyncMessage;
+class FAsyncMessageSystemBase;
+
+// Note: The payloads need to be USTRUCTs in order to test them with FInstancedStruct, which require ScriptStructs.
+
+/**
+ * A simple test payload type to ensure that the data of the payload is correct when you recieve a message.
+ */
+USTRUCT(NotBlueprintType, meta=(Hidden))
 struct FTest_Payload_A
 {
 	GENERATED_BODY()
@@ -20,76 +30,33 @@ struct FTest_Payload_A
 
 	UPROPERTY()
 	float Bar = 123.0f;
+
+	UPROPERTY()
+	TObjectPtr<UObject> TestPointer = nullptr;
 };
 
-struct FAsyncMessage;
+/**
+ * A test payload to allow for us to check that you can bind listeners within the response
+ * to another message, a "Nested" message binding.
+ */
+USTRUCT(NotBlueprintType, meta=(Hidden))
+struct FNested_Payload
+{
+	GENERATED_BODY()
+	
+	TWeakPtr<FAsyncMessageSystemBase> MessageSystem = nullptr;
+};
 
-UCLASS(NotBlueprintable)
+/**
+ * A test UObject type which we can use to ensure that you can bind messages to UObjects.
+ */
+UCLASS(NotBlueprintable, NotBlueprintType)
 class UTestAsyncObject : public UObject
 {
 	GENERATED_BODY()
-
 public:
-	
 	void CallbackFunction(const FAsyncMessage& Message);
 
 	UPROPERTY()
 	int32 TestValue = 5;
-	
-};
-
-namespace UE::AsyncMessageSystem
-{
-	// Test message system interface impl
-	class FTestMessageSystem : public FAsyncMessageSystemBase
-	{
-	protected:
-		virtual void Startup_Impl() override
-		{
-			bHasStarted = true;
-		}
-
-		virtual void Shutdown_Impl() override
-		{
-			bHasShutdown = true;
-		}
-
-		virtual void PostQueueMessage(const FAsyncMessageId MessageId, const TSet<FAsyncMessageBindingOptions>& OptionsBoundTo) override
-		{
-			// Dont do anything by default here in this simple test message system
-		}
-		
-		bool bHasStarted = false;
-		bool bHasShutdown = false;	
-
-	public:
-		const bool HasStarted() const { return bHasStarted; }
-		const bool HasShutdown() const { return bHasShutdown && bIsShuttingDown; }
-		
-		/**
-		 * Returns true if the handle currently has any listeners bound to it in the message map
-		 */
-		bool IsHandleBound(const FAsyncMessageHandle& Handle) const;
-
-		uint32 GetNumberOfListeners() const;
-		
-		void RunOnce()
-		{
-			ProcessMessageQueueForBinding(FAsyncMessageBindingOptions{});
-		}
-
-		void Test_Shutdown()
-		{
-			Shutdown();
-		}
-
-		FAsyncMessageHandle GenerateTestHandle();
-	};
-}
-
-USTRUCT()
-struct FNested_Payload
-{
-	GENERATED_BODY()
-	TWeakPtr<UE::AsyncMessageSystem::FTestMessageSystem> MessageSystem = nullptr;
 };
\ No newline at end of file
