Fix up some incorrectly name trace scopes on the message system. Move the AsyncMessageStore.h file to public folder to make it easier to depend on the module

#rnx

diff --git a/Engine/Plugins/Experimental/AsyncMessageSystem/Source/AsyncMessageSystem/Private/AsyncMessageStore.cpp b/Engine/Plugins/Experimental/AsyncMessageSystem/Source/AsyncMessageSystem/Private/AsyncMessageStore.cpp
--- a/Engine/Plugins/Experimental/AsyncMessageSystem/Source/AsyncMessageSystem/Private/AsyncMessageStore.cpp
+++ b/Engine/Plugins/Experimental/AsyncMessageSystem/Source/AsyncMessageSystem/Private/AsyncMessageStore.cpp
@@ -4,7 +4,7 @@
 
 uint32 FAsyncMessageStore::EnqueueMessage(FAsyncMessage&& MessageToQueue, TSet<FAsyncMessageBindingOptions>& BindingTypesAddedTo)
 {
-	TRACE_CPUPROFILER_EVENT_SCOPE(IAsyncMessageSystem::ProcessMessageQueueForBinding_Impl);
+	TRACE_CPUPROFILER_EVENT_SCOPE(FAsyncMessageStore::EnqueueMessage);
 	
 	// We want to keep track of which BindingOptions had listeners for this message ID
 	BindingTypesAddedTo.Reset();
diff --git a/Engine/Plugins/Experimental/AsyncMessageSystem/Source/AsyncMessageSystem/Private/AsyncMessageStore.h b/Engine/Plugins/Experimental/AsyncMessageSystem/Source/AsyncMessageSystem/Private/AsyncMessageStore.h
deleted file mode 100644
diff --git a/Engine/Plugins/Experimental/AsyncMessageSystem/Source/AsyncMessageSystem/Private/AsyncMessageSystemBase.cpp b/Engine/Plugins/Experimental/AsyncMessageSystem/Source/AsyncMessageSystem/Private/AsyncMessageSystemBase.cpp
--- a/Engine/Plugins/Experimental/AsyncMessageSystem/Source/AsyncMessageSystem/Private/AsyncMessageSystemBase.cpp
+++ b/Engine/Plugins/Experimental/AsyncMessageSystem/Source/AsyncMessageSystem/Private/AsyncMessageSystemBase.cpp
@@ -32,13 +32,13 @@ void FAsyncMessageSystemBase::UnbindListener(const FAsyncMessageHandle& HandleTo
 
 bool FAsyncMessageSystemBase::QueueMessageForBroadcast(const FAsyncMessageId MessageId, FConstStructView PayloadData)
 {
-	TRACE_CPUPROFILER_EVENT_SCOPE(IAsyncMessageSystem::QueueMessageForBroadcast);
+	TRACE_CPUPROFILER_EVENT_SCOPE(FAsyncMessageSystemBase::QueueMessageForBroadcast);
 	return QueueMessageForBroadcast_Impl(MessageId, MoveTemp(PayloadData));
 }
 
 void FAsyncMessageSystemBase::ProcessMessageQueueForBinding(const FAsyncMessageBindingOptions& Options)
 {
-	TRACE_CPUPROFILER_EVENT_SCOPE(IAsyncMessageSystem::ProcessMessageQueueForBinding);
+	TRACE_CPUPROFILER_EVENT_SCOPE(FAsyncMessageSystemBase::ProcessMessageQueueForBinding);
 
 	// Handle any pre-processing which may need to happen before sending messages. This could include
 	// binding listeners, removing old handles, and more.
@@ -50,7 +50,7 @@ void FAsyncMessageSystemBase::ProcessMessageQueueForBinding(const FAsyncMessageB
 
 void FAsyncMessageSystemBase::Startup()
 {
-	TRACE_CPUPROFILER_EVENT_SCOPE(IAsyncMessageSystem::Startup);
+	TRACE_CPUPROFILER_EVENT_SCOPE(FAsyncMessageSystemBase::Startup);
 	
 	UE_LOG(LogAsyncMessageSystem, Verbose, TEXT("[%hs] Message System Startup..."), __func__);
 	
@@ -61,7 +61,7 @@ void FAsyncMessageSystemBase::Startup()
 
 void FAsyncMessageSystemBase::Shutdown()
 {
-	TRACE_CPUPROFILER_EVENT_SCOPE(IAsyncMessageSystem::Shutdown);
+	TRACE_CPUPROFILER_EVENT_SCOPE(FAsyncMessageSystemBase::Shutdown);
 	
 	UE_LOG(LogAsyncMessageSystem, Verbose, TEXT("[%hs] Message System Shutdown..."), __func__);
 	
@@ -77,7 +77,7 @@ void FAsyncMessageSystemBase::Shutdown()
 
 bool FAsyncMessageSystemBase::BindListener_Impl(const FAsyncMessageHandle& HandleToBindTo, const FAsyncMessageId MessageId, FMessageCallbackFunc&& Callback, const FAsyncMessageBindingOptions& Options)
 {
-	TRACE_CPUPROFILER_EVENT_SCOPE(IAsyncMessageSystem::BindListener_Impl);
+	TRACE_CPUPROFILER_EVENT_SCOPE(FAsyncMessageSystemBase::BindListener_Impl);
 	
 	// Do not allow the binding of new listeners during the shutdown of this message system
 	if (bIsShuttingDown)
@@ -108,18 +108,12 @@ bool FAsyncMessageSystemBase::BindListener_Impl(const FAsyncMessageHandle& Handl
 	// Let the message store know that this message will be listened for by this binding option
 	MessageStore.AddMessageToBinding(MessageId, Options);
 
-	FAsyncMessageIndividualListener<FMessageCallbackFunc> ListenerData
-		{
-			.Handle = HandleToBindTo,
-			.Callback = MoveTemp(Callback)
-		};
-
 	// Queue this listener for binding the next time that this system is processed.
 	PendingBoundListenerQueue.Enqueue(
 		FPendingBoundListener<FMessageCallbackFunc> 
 		{
 			.MessageId = MessageId, 
-			.Data = MoveTemp(ListenerData), 
+			.Data = FAsyncMessageIndividualListener<FMessageCallbackFunc> { .Handle = HandleToBindTo, .Callback = MoveTempIfPossible(Callback) }, 
 			.Options = Options 
 		});
 
@@ -195,12 +189,12 @@ bool FAsyncMessageSystemBase::QueueMessageForBroadcast_Impl(const FAsyncMessageI
 	
 	// Return true if there were any bound listeners to this message. 
 	// False if the message was not added to any message queues
-	return NumQueuesAddedTo > 0;
+	return NumQueuesAddedTo > 0u;
 }
 
 void FAsyncMessageSystemBase::PreProcessMessagesQueue(const FAsyncMessageBindingOptions& Options)
 {
-	TRACE_CPUPROFILER_EVENT_SCOPE(IAsyncMessageSystem::PreProcessMessagesQueue);
+	TRACE_CPUPROFILER_EVENT_SCOPE(FAsyncMessageSystemBase::PreProcessMessagesQueue);
 
 	// Remove any message handles that have been unbound
 	ProcessUnbindHandleRequests();
@@ -211,22 +205,24 @@ void FAsyncMessageSystemBase::PreProcessMessagesQueue(const FAsyncMessageBinding
 
 void FAsyncMessageSystemBase::ProcessMessageQueueForBinding_Impl(const FAsyncMessageBindingOptions& Options)
 {
-	TRACE_CPUPROFILER_EVENT_SCOPE(IAsyncMessageSystem::ProcessMessageQueueForBinding_Impl);
+	TRACE_CPUPROFILER_EVENT_SCOPE(FAsyncMessageSystemBase::ProcessMessageQueueForBinding_Impl);
 	
 	// Reset the message sequence ID which we are processing this frame
-	NextMessageSequence = 0;
+	NextMessageSequence = 0u;
 
 	// Lock the message listener maps so that listeners don't get added/removed
 	// in the middle of processing, which would cause a data race
 	FScopeLock ListenerLock(&MessageListenerMapCS);
 	
-	uint32 ProcessedMessages = 0;
-	uint32 MessagesCalled = 0;
+	uint32 ProcessedMessages = 0u;
+	uint32 MessagesCalled = 0u;
 	
 	// Pop off of the message queue and process each message
 	TOptional<FAsyncMessage> OptionalNextMessage = MessageStore.GetNextMessageForBindingOption(Options);
 	while (OptionalNextMessage.IsSet())
 	{
+		TRACE_CPUPROFILER_EVENT_SCOPE(FAsyncMessageSystemBase::ProcessMessageQueueForBinding_Impl::SingleMessage);
+
 		FAsyncMessage& Message = OptionalNextMessage.GetValue();
 
 		// TODO: If we define a "lifetime" for messages, we could grab the current frame/time
@@ -236,7 +232,7 @@ void FAsyncMessageSystemBase::ProcessMessageQueueForBinding_Impl(const FAsyncMes
 		//const double CurrentTime = FApp::GetCurrentTime();
 		//const uint64 CurrentFrame = GFrameCounter;
 		
-		// For each message in the heirarchy of this message...
+		// For each message in the hierarchy of this message...
 		FAsyncMessageId::WalkMessageHierarchy(Message.GetMessageId(), [&](const FAsyncMessageId CurrentMessageId)
 		{
 			// Set this message ID to the current message id
@@ -278,7 +274,7 @@ void FAsyncMessageSystemBase::ProcessMessageQueueForBinding_Impl(const FAsyncMes
 
 void FAsyncMessageSystemBase::ProcessUnbindHandleRequests()
 {
-	TRACE_CPUPROFILER_EVENT_SCOPE(IAsyncMessageSystem::ProcessUnbindHandleRequests);
+	TRACE_CPUPROFILER_EVENT_SCOPE(FAsyncMessageSystemBase::ProcessUnbindHandleRequests);
 
 	if (MessageHandlesPendingRemoval.IsEmpty())
 	{
@@ -318,7 +314,7 @@ void FAsyncMessageSystemBase::ProcessUnbindHandleRequests()
 
 void FAsyncMessageSystemBase::ProcessListenersPendingBinding()
 {
-	TRACE_CPUPROFILER_EVENT_SCOPE(IAsyncMessageSystem::ProcessListenersPendingBinding);
+	TRACE_CPUPROFILER_EVENT_SCOPE(FAsyncMessageSystemBase::ProcessListenersPendingBinding);
 
 	if (PendingBoundListenerQueue.IsEmpty())
 	{
diff --git a/Engine/Plugins/Experimental/AsyncMessageSystem/Source/AsyncMessageSystem/Public/AsyncMessageStore.h b/Engine/Plugins/Experimental/AsyncMessageSystem/Source/AsyncMessageSystem/Public/AsyncMessageStore.h
new file mode 100644
--- /dev/null
+++ b/Engine/Plugins/Experimental/AsyncMessageSystem/Source/AsyncMessageSystem/Public/AsyncMessageStore.h
@@ -0,0 +1,78 @@
+// Copyright Epic Games, Inc. All Rights Reserved.
+
+#pragma once
+
+#include "AsyncMessage.h"
+#include "AsyncMessageBindingOptions.h"
+#include "AsyncMessageId.h"
+#include "Containers/Map.h"
+#include "Containers/Queue.h"
+#include "Templates/SharedPointer.h"
+
+/**
+ * Implementation for the actual storage of async messages for different binding methods.
+ * 
+ * This will store a message queue for each type of supports Binding option for listeners,
+ * allowing listeners to be in control over when they receive messages from the queue (what tick group,
+ * thread ID, etc).
+ */
+class FAsyncMessageStore final
+{
+public:
+	FAsyncMessageStore() = default;
+	~FAsyncMessageStore() = default;
+
+	UE_NONCOPYABLE(FAsyncMessageStore);
+
+public:
+
+	uint32 EnqueueMessage(FAsyncMessage&& MessageToQueue, OUT TSet<FAsyncMessageBindingOptions>& BindingTypesAddedTo);
+
+	[[nodiscard]] TOptional<FAsyncMessage> GetNextMessageForBindingOption(const FAsyncMessageBindingOptions& Binding);
+
+	void AddMessageToBinding(const FAsyncMessageId& MessageId, const FAsyncMessageBindingOptions& Binding);
+
+	/**
+	 * Removes the given message ID from the specified message binding. This is called when an message handler is unbound
+	 * 
+	 */
+	bool RemoveMessageFromBinding(const FAsyncMessageId& MessageId, const FAsyncMessageBindingOptions& Binding);
+
+protected:
+
+	/**
+	* CS for when access to the MessageBindingQueues is required 
+	* to enqueue or dequeue messages accordingly.
+	*/
+	mutable FCriticalSection MessageQueueCriticalSection;
+	
+	// The message queue for a single set of binding options
+	struct FBindingOptionsMessageQueue
+	{
+		FBindingOptionsMessageQueue()
+		{
+			MessageQueue = MakeShared<TQueue<FAsyncMessage>>();
+		}
+
+		~FBindingOptionsMessageQueue()
+		{
+			MessageQueue.Reset();
+		}
+		
+		// Keep track of what messages are bound within this interface so that we can quickly
+		// find which queues require which messages.
+		// The value of this map is how many times it has been bound.
+		TMap<FAsyncMessageId, uint32> MessagesWithTheseBindings;
+
+		// The queue of messages for this binding option		
+		TSharedPtr<TQueue<FAsyncMessage>> MessageQueue;
+	};
+
+	/**
+	 * Tracks the binding options to their associated message queues.
+	 * 
+	 * Each binding option needs its own message queue so that we can enable
+	 * the listeners of the messages to control _when_ they receive an message during the frame. 
+	 */
+	TMap<FAsyncMessageBindingOptions, FBindingOptionsMessageQueue> MessageBindingQueues;
+};
