Add a check for if the system is shutting down before beginning processing of messages. Avoid adding invalid handles to the removal queue which would just cause needless work

#rnx
#prelight todo

diff --git a/Engine/Plugins/Experimental/AsyncMessageSystem/Source/AsyncMessageSystem/Private/AsyncMessageSystemBase.cpp b/Engine/Plugins/Experimental/AsyncMessageSystem/Source/AsyncMessageSystem/Private/AsyncMessageSystemBase.cpp
--- a/Engine/Plugins/Experimental/AsyncMessageSystem/Source/AsyncMessageSystem/Private/AsyncMessageSystemBase.cpp
+++ b/Engine/Plugins/Experimental/AsyncMessageSystem/Source/AsyncMessageSystem/Private/AsyncMessageSystemBase.cpp
@@ -122,8 +122,11 @@ bool FAsyncMessageSystemBase::BindListener_Impl(const FAsyncMessageHandle& Handl
 }
 
 void FAsyncMessageSystemBase::UnbindListener_Impl(const FAsyncMessageHandle& HandleToUnbind)
-{	
-	MessageHandlesPendingRemoval.Enqueue(HandleToUnbind);
+{
+	if (HandleToUnbind.IsValid())
+	{
+		MessageHandlesPendingRemoval.Enqueue(HandleToUnbind);	
+	}
 }
 
 bool FAsyncMessageSystemBase::QueueMessageForBroadcast_Impl(const FAsyncMessageId MessageId, FConstStructView PayloadData)
@@ -206,6 +209,12 @@ void FAsyncMessageSystemBase::PreProcessMessagesQueue(const FAsyncMessageBinding
 void FAsyncMessageSystemBase::ProcessMessageQueueForBinding_Impl(const FAsyncMessageBindingOptions& Options)
 {
 	TRACE_CPUPROFILER_EVENT_SCOPE(FAsyncMessageSystemBase::ProcessMessageQueueForBinding_Impl);
+
+	if (bIsShuttingDown)
+	{
+		UE_LOG(LogAsyncMessageSystem, Error, TEXT("[%hs] Attempting to Process Message Queue during shutdown, exiting"), __func__);
+		return;
+	}
 	
 	// Reset the message sequence ID which we are processing this frame
 	NextMessageSequence = 0u;
