Blueprint support for Async Message System. You can now call "Start Listening for Async Message" in a blueprint and give it a message ID to listen for, and you will receive a callback when it happens. You can also Queue messages for broadcasting

#rnx

diff --git a/Engine/Plugins/Experimental/AsyncMessageSystem/Source/AsyncMessageSystem/Private/AsyncAction_ListenForMessage.cpp b/Engine/Plugins/Experimental/AsyncMessageSystem/Source/AsyncMessageSystem/Private/AsyncAction_ListenForMessage.cpp
new file mode 100644
--- /dev/null
+++ b/Engine/Plugins/Experimental/AsyncMessageSystem/Source/AsyncMessageSystem/Private/AsyncAction_ListenForMessage.cpp
@@ -0,0 +1,158 @@
+// Copyright Epic Games, Inc. All Rights Reserved.
+
+#include "AsyncAction_ListenForMessage.h"
+
+#include "AsyncGameplayMessageSystem.h"
+#include "AsyncMessageWorldSubsystem.h"
+#include "AsyncMessageSystemLogs.h"
+#include "Engine/Engine.h"
+
+UAsyncAction_ListenForAsyncMessage* UAsyncAction_ListenForAsyncMessage::StartListeningForAsyncMessage(
+	UObject* WorldContextObject,
+	const FAsyncMessageId MessageId,
+	const TEnumAsByte<ETickingGroup> TickGroup)
+{
+	check(GEngine);
+	
+	// Use the GEngine version of GetWorld instead of a simpler UObject function call because it has nice
+	// log/error raising for BP users.
+	UWorld* World = GEngine->GetWorldFromContextObject(WorldContextObject, EGetWorldErrorMode::LogAndReturnNull);
+	if (!World)
+	{
+		return nullptr;
+	}
+
+	// Create a new Async BP action to listen for a message
+	UAsyncAction_ListenForAsyncMessage* Action = NewObject<UAsyncAction_ListenForAsyncMessage>();
+	Action->WeakWorldPtr = World;
+	Action->MessageToListenFor = MessageId;
+	Action->BindingOptions.SetTickGroup(TickGroup);
+	
+	Action->RegisterWithGameInstance(World);
+	
+	return Action;
+}
+
+void UAsyncAction_ListenForAsyncMessage::StopListeningForAsyncMessage()
+{
+	// Mark this async action as being no long needed, which will unbind our listeners from the message system
+	SetReadyToDestroy();
+}
+
+void UAsyncAction_ListenForAsyncMessage::Activate()
+{
+	Super::Activate();
+
+	StartListeningForMessage();
+}
+
+void UAsyncAction_ListenForAsyncMessage::SetReadyToDestroy()
+{
+	Super::SetReadyToDestroy();
+
+	UnbindListener();
+}
+
+void UAsyncAction_ListenForAsyncMessage::StartListeningForMessage()
+{
+	check (!BoundListenerHandle.IsValid());
+	
+	if (!MessageToListenFor.IsValid())
+	{
+		UE_LOG(LogAsyncMessageSystem, Error, TEXT("[%hs] Invalid MessageToListenFor '%s' on action %s"), __func__, *MessageToListenFor.ToString(), *GetNameSafe(this));
+		return;
+	}
+
+	TSharedPtr<FAsyncGameplayMessageSystem> Sys = GetAssociatedMessageSystem();
+	if (!Sys.IsValid())
+	{
+		UE_LOG(LogAsyncMessageSystem, Error, TEXT("[%hs] Failed to find the associated message system for async action %s"), __func__, *GetNameSafe(this));
+		return;
+	}
+
+	// Bind the listener for this message
+	BoundListenerHandle = Sys->BindListener(
+		MessageToListenFor,
+		TWeakObjectPtr<UAsyncAction_ListenForAsyncMessage>{ this },
+		&UAsyncAction_ListenForAsyncMessage::HandleMessageReceived,
+		BindingOptions);
+
+	// If for some reason we failed to bind the handle, then destroy this async action and log an error
+	if (!BoundListenerHandle.IsValid())
+	{
+		UE_LOG(LogAsyncMessageSystem, Error, TEXT("[%hs] Failed to bind listener for async action %s"), __func__, *GetNameSafe(this));
+		
+		SetReadyToDestroy();
+		return;
+	}
+
+	// Listen for when the subsystem starts to shut down so that we can unbind ourselves and clean up properly
+	UAsyncMessageWorldSubsystem* WorldSubsystem = GetAssociatedSubsystem();
+	if (!WorldSubsystem)
+	{
+		UE_LOG(LogAsyncMessageSystem, Error, TEXT("[%hs] Failed to find async message world subsystem %s"), __func__, *GetNameSafe(this));
+		SetReadyToDestroy();
+		return;
+	}
+	
+	check(!OnMessageSystemShutdownDelegateHandle.IsValid());
+	OnMessageSystemShutdownDelegateHandle = WorldSubsystem->OnShutdownMessageSystem.AddWeakLambda(this, [this]()
+	{
+		// Mark ourselves as being ready for destruction when the associated message subsystem goes out of scope
+		SetReadyToDestroy();
+	});
+}
+
+void UAsyncAction_ListenForAsyncMessage::UnbindListener()
+{
+	// Stop listening to the subsystem shutdown delegate
+	if (OnMessageSystemShutdownDelegateHandle.IsValid())
+	{
+		if (TObjectPtr<UAsyncMessageWorldSubsystem> WorldSubsys = GetAssociatedSubsystem())
+		{
+			WorldSubsys->OnShutdownMessageSystem.Remove(OnMessageSystemShutdownDelegateHandle);
+		}
+	}
+
+	// Stop listening for the actual message callback
+	if (!BoundListenerHandle.IsValid())
+	{
+		return;
+	}
+	
+	TSharedPtr<FAsyncGameplayMessageSystem> Sys = GetAssociatedMessageSystem();
+	if (!Sys.IsValid())
+	{
+		return;
+	}
+
+	// Unbind this async action listener
+	Sys->UnbindListener(BoundListenerHandle);
+}
+
+void UAsyncAction_ListenForAsyncMessage::HandleMessageReceived(const FAsyncMessage& Message)
+{
+	OnMessageReceived.Broadcast(Message);
+}
+
+TSharedPtr<FAsyncGameplayMessageSystem> UAsyncAction_ListenForAsyncMessage::GetAssociatedMessageSystem() const
+{
+	TStrongObjectPtr<UWorld> World = WeakWorldPtr.Pin();
+	if (!World.IsValid())
+	{
+		return nullptr;
+	}
+
+	return UAsyncMessageWorldSubsystem::GetSharedMessageSystem<FAsyncGameplayMessageSystem>(World.Get());
+}
+
+TObjectPtr<UAsyncMessageWorldSubsystem> UAsyncAction_ListenForAsyncMessage::GetAssociatedSubsystem() const
+{
+	TStrongObjectPtr<UWorld> World = WeakWorldPtr.Pin();
+	if (!World.IsValid())
+	{
+		return nullptr;
+	}
+
+	return World->GetSubsystem<UAsyncMessageWorldSubsystem>();
+}
diff --git a/Engine/Plugins/Experimental/AsyncMessageSystem/Source/AsyncMessageSystem/Private/AsyncAction_ListenForMessage.h b/Engine/Plugins/Experimental/AsyncMessageSystem/Source/AsyncMessageSystem/Private/AsyncAction_ListenForMessage.h
new file mode 100644
--- /dev/null
+++ b/Engine/Plugins/Experimental/AsyncMessageSystem/Source/AsyncMessageSystem/Private/AsyncAction_ListenForMessage.h
@@ -0,0 +1,122 @@
+// Copyright Epic Games, Inc. All Rights Reserved.
+
+#pragma once
+
+#include "AsyncMessageBindingOptions.h"
+#include "AsyncMessageHandle.h"
+#include "AsyncMessageId.h"
+#include "Kismet/BlueprintAsyncActionBase.h"
+#include "StructUtils/InstancedStruct.h"
+
+#include "AsyncAction_ListenForMessage.generated.h"
+
+struct FAsyncMessage;
+class FAsyncGameplayMessageSystem;
+class UAsyncMessageWorldSubsystem;
+
+DECLARE_DYNAMIC_MULTICAST_DELEGATE_OneParam(FAsyncMessageReceivedDelegate, const FAsyncMessage&, Message);
+
+/**
+ * An async action for binding a listener to a message. The user will call the
+ * "StartListeningForAsyncMessage" function from blueprints, which will create the async action object
+ * and an exec pin will be created for any BP Assignable delegates on this object (FAsyncMessageReceivedDelegate).
+ *
+ * Those delegates can then fire when the message is received, and we can allow users to specify their own tick groups upon binding as well.
+ * We expose this as an async proxy (ExposedAsyncProxy) because then you can easily call the "StopListeningForAsyncMessage" function
+ * to unbind the listener.
+ *
+ * If for some reason there is a failure when binding to the message, this async task will be immedately marked as being ready for destruction.
+ */
+UCLASS(BlueprintType, meta=(ExposedAsyncProxy = "AsyncMessageTask"))
+class UAsyncAction_ListenForAsyncMessage : public UBlueprintAsyncActionBase
+{
+	GENERATED_BODY()
+	
+public:
+	
+	/**
+	 * Starts listening for an Async Message with the given ID during the given tick group. 
+	 * 
+	 * @param MessageId The ID of the Async Message which you would like to listen to
+	 * @param TickGroup The tick group which you would like to receive the message in. Default is TG_PostUpdateWork.
+	 */
+	UFUNCTION(BlueprintCallable, Category = "Async Messages", meta=(WorldContext="WorldContextObject", BlueprintInternalUseOnly="true", AdvancedDisplay="TickGroup", DisplayName="Start Listening for Async Message"))
+	static UAsyncAction_ListenForAsyncMessage* StartListeningForAsyncMessage(
+		UObject* WorldContextObject,
+		const FAsyncMessageId MessageId,
+		const TEnumAsByte<ETickingGroup> TickGroup = TG_PostUpdateWork);
+
+	/**
+	 * Stops the given message handle from receiving any more messages.
+	 */
+	UFUNCTION(BlueprintCallable, Category = "Async Messages")
+	void StopListeningForAsyncMessage();
+
+	/**
+	 * Delegate which executes when the bound message that this async action is listening for is broadcast.
+	 */
+	UPROPERTY(BlueprintAssignable)
+	FAsyncMessageReceivedDelegate OnMessageReceived;
+	
+protected:
+
+	//~ Begin UBlueprintAsyncActionBase interface
+	virtual void Activate() override;
+	virtual void SetReadyToDestroy() override;
+	//~ End of UBlueprintAsyncActionBase interface
+
+	/**
+	 * Binds a listener to for the MessageToListenFor message id to the HandleMessageReceived function 
+	 */
+	void StartListeningForMessage();
+
+	/**
+	 * Unbinds the listener handle for this async action so no more messages will be processed by it.
+	 */
+	void UnbindListener();
+
+	/**
+	 * The callback function which will call the OnMessageReceived delegate and give blueprints
+	 * the oppurtunity to process the message.
+	 */
+	void HandleMessageReceived(const FAsyncMessage& Message);
+
+	/**
+	 * @return The gameplay message system associated with the outer world of this async action.
+	 */
+	TSharedPtr<FAsyncGameplayMessageSystem> GetAssociatedMessageSystem() const;
+
+	/**
+	 * @return The async message world subsystem associated with the outer world of this async action. 
+	 */
+	TObjectPtr<UAsyncMessageWorldSubsystem> GetAssociatedSubsystem() const;
+	
+	/**
+	* Weak pointer to the owning wold context to which this listener belongs 
+	*/
+	TWeakObjectPtr<UWorld> WeakWorldPtr = nullptr;
+	
+	/**
+	* The async message id that this action should listen 
+	*/
+	FAsyncMessageId MessageToListenFor = FAsyncMessageId::Invalid;
+
+	/**
+	 * The binding options to use when listening for this message.
+	 * 
+	 * The tick group of this binding option is set on construction of this object in StartListeningForAsyncMessage.
+	 */
+	FAsyncMessageBindingOptions BindingOptions = {};
+
+	/**
+	 * The listener handle which has been bound to the message 
+	 */
+	FAsyncMessageHandle BoundListenerHandle = FAsyncMessageHandle::Invalid;
+
+	/**
+	 * Handle to a delegate which fires when the world subsystem which owns the message system
+	 * that this async action is listening for shuts down. This allows us to clean up this async
+	 * action and mark it as being ready for destruction.
+	 */
+	FDelegateHandle OnMessageSystemShutdownDelegateHandle;
+};
\ No newline at end of file
diff --git a/Engine/Plugins/Experimental/AsyncMessageSystem/Source/AsyncMessageSystem/Private/AsyncMessageBlueprintLibrary.cpp b/Engine/Plugins/Experimental/AsyncMessageSystem/Source/AsyncMessageSystem/Private/AsyncMessageBlueprintLibrary.cpp
new file mode 100644
--- /dev/null
+++ b/Engine/Plugins/Experimental/AsyncMessageSystem/Source/AsyncMessageSystem/Private/AsyncMessageBlueprintLibrary.cpp
@@ -0,0 +1,40 @@
+// Copyright Epic Games, Inc. All Rights Reserved.
+
+#include "AsyncMessageBlueprintLibrary.h"
+
+#include "AsyncMessageWorldSubsystem.h"
+#include "AsyncMessageSystemBase.h"
+#include "AsyncMessageSystemLogs.h"
+#include "Engine/Engine.h"	// For GEnngine 
+
+bool UAsyncMessageSystemBlueprintLibrary::QueueAsyncMessageForBroadcast(
+	UObject* WorldContextObject,
+	const FAsyncMessageId& MessageId,
+	const FInstancedStruct& Payload)
+{
+	check(GEngine);
+	UWorld* World = GEngine->GetWorldFromContextObject(WorldContextObject, EGetWorldErrorMode::LogAndReturnNull);
+	if (!World)
+	{
+		UE_LOG(LogAsyncMessageSystem, Error, TEXT("[%hs] Failed to queue message '%s' for broadcasting: Unable to find a world."),
+			__func__, *MessageId.ToString());
+		
+		return false;
+	}
+
+	TSharedPtr<FAsyncMessageSystemBase> Sys = UAsyncMessageWorldSubsystem::GetSharedMessageSystem(World);
+	if (!Sys.IsValid())
+	{
+		UE_LOG(LogAsyncMessageSystem, Error, TEXT("[%hs] Failed to queue message '%s' for broadcasting: Unable to find a message system for world '%s'."),
+			__func__, *MessageId.ToString(),  *GetNameSafe(World));
+		
+		return false;
+	}
+
+	return Sys->QueueMessageForBroadcast(MessageId, Payload);
+}
+
+FString UAsyncMessageSystemBlueprintLibrary::Conv_AsyncMessageIdToString(const FAsyncMessageId& MessageId)
+{
+	return MessageId.ToString();
+}
diff --git a/Engine/Plugins/Experimental/AsyncMessageSystem/Source/AsyncMessageSystem/Private/AsyncMessageBlueprintLibrary.h b/Engine/Plugins/Experimental/AsyncMessageSystem/Source/AsyncMessageSystem/Private/AsyncMessageBlueprintLibrary.h
new file mode 100644
--- /dev/null
+++ b/Engine/Plugins/Experimental/AsyncMessageSystem/Source/AsyncMessageSystem/Private/AsyncMessageBlueprintLibrary.h
@@ -0,0 +1,45 @@
+// Copyright Epic Games, Inc. All Rights Reserved.
+
+#pragma once
+
+#include "AsyncMessageId.h"
+#include "InstancedStruct.h"
+#include "Kismet/BlueprintFunctionLibrary.h"
+
+#include "AsyncMessageBlueprintLibrary.generated.h"
+
+/**
+ * Blueprint function library for the Async Message System
+ */
+UCLASS()
+class UAsyncMessageSystemBlueprintLibrary final : public UBlueprintFunctionLibrary
+{
+	GENERATED_BODY()
+
+public:
+
+	/**
+	 * Queues the given async message for broadcast the next time that this
+	 * message system processes its message queue
+	 *
+	 * @param MessageId		The message that you would like to queue for broadcasting
+	 * @param Payload		The payload data of this message. This payload data will be COPIED to the message queue to
+	 *						to make it safe for listeners on other threads. 
+	 *
+	 * @return True if this message had an listeners bound to it and it was successfully queued
+	 */
+	UFUNCTION(BlueprintCallable, Category = "Async Messages", meta=(WorldContext="WorldContextObject", ReturnDisplayName = "Success"))
+	static bool QueueAsyncMessageForBroadcast(
+		UObject* WorldContextObject,
+		const FAsyncMessageId& MessageId,
+		const FInstancedStruct& Payload);
+
+	/**
+	 * Get the string representation of the given Async Message Id.
+	 * 
+	 * @param MessageId The Message Id to get the string representation of
+	 * @return The String representation of the give async message Id
+	 */
+	UFUNCTION(BlueprintPure, Category = "Utilities|String", meta=(DisplayName = "To String (Async Message Id)", CompactNodeTitle = "->", BlueprintAutocast))
+	static FString Conv_AsyncMessageIdToString(const FAsyncMessageId& MessageId);
+};
diff --git a/Engine/Plugins/Experimental/AsyncMessageSystem/Source/AsyncMessageSystem/Private/AsyncMessageWorldSubsystem.cpp b/Engine/Plugins/Experimental/AsyncMessageSystem/Source/AsyncMessageSystem/Private/AsyncMessageWorldSubsystem.cpp
--- a/Engine/Plugins/Experimental/AsyncMessageSystem/Source/AsyncMessageSystem/Private/AsyncMessageWorldSubsystem.cpp
+++ b/Engine/Plugins/Experimental/AsyncMessageSystem/Source/AsyncMessageSystem/Private/AsyncMessageWorldSubsystem.cpp
@@ -36,6 +36,8 @@ void UAsyncMessageWorldSubsystem::ShutdownMessageSystem()
 		
 		MessageSystem->Shutdown();
 		
-		MessageSystem.Reset();
+		MessageSystem.Reset();		
 	}
+	
+	OnShutdownMessageSystem.Broadcast();
 }
diff --git a/Engine/Plugins/Experimental/AsyncMessageSystem/Source/AsyncMessageSystem/Public/AsyncMessage.h b/Engine/Plugins/Experimental/AsyncMessageSystem/Source/AsyncMessageSystem/Public/AsyncMessage.h
--- a/Engine/Plugins/Experimental/AsyncMessageSystem/Source/AsyncMessageSystem/Public/AsyncMessage.h
+++ b/Engine/Plugins/Experimental/AsyncMessageSystem/Source/AsyncMessageSystem/Public/AsyncMessage.h
@@ -6,6 +6,8 @@
 #include "HAL/Platform.h"	// For uint64 type declaration
 #include "StructUtils/StructView.h"
 
+#include "AsyncMessage.generated.h"
+
 /**
  * A single async message which has been queued for broadcasting, and listeners are
  * interested in receiving the data from.
@@ -13,8 +15,10 @@
  * Async Message Messages contain some data about when, where, and how the message was queued
  * which make it possible to synchronize game state across multiple threads with messages. 
  */
+USTRUCT(BlueprintType)
 struct ASYNCMESSAGESYSTEM_API FAsyncMessage final
 {
+	GENERATED_BODY()
 public:
 	
 	FAsyncMessage() = default;
@@ -91,6 +95,7 @@ private:
 	/**
 	 * The Async message Id which this message is a part of. 
 	 */
+	UPROPERTY(BlueprintReadOnly, VisibleAnywhere, Category="Async Message", meta = (AllowPrivateAccess = "true"))
 	FAsyncMessageId MessageId = FAsyncMessageId::Invalid;
 
 	/**
@@ -103,11 +108,13 @@ private:
 	 *
 	 * If the MessageSourceId and MessageId are the same value, then that means that this is the childmost message.
 	 */
+	UPROPERTY(BlueprintReadOnly, VisibleAnywhere, Category="Async Message", meta = (AllowPrivateAccess = "true"))
 	FAsyncMessageId MessageSourceId = FAsyncMessageId::Invalid;
 
 	/**
 	 * The time at which this message was queued ( FApp::GetCurrentTime() )
 	 */
+	UPROPERTY(BlueprintReadOnly, VisibleAnywhere, Category="Async Message", meta = (AllowPrivateAccess = "true"))
 	double QueueTime = 0.0;
 
 	/**
@@ -142,6 +149,7 @@ private:
 	 * TODO: When we copy the payload data, have it use a custom linear allocator on the owning message
 	 * system to make it faster to copy and more cache friendly to the message system to iterate messages.
 	 */
+	UPROPERTY(BlueprintReadOnly, VisibleAnywhere, Category="Async Message", meta = (DisplayName="Payload", AllowPrivateAccess = "true"))
 	FInstancedStruct PayloadCopy;
 
 	// TODO: We could perhaps define a "lifetime" concept on these messages, similar to how AnimNext's
diff --git a/Engine/Plugins/Experimental/AsyncMessageSystem/Source/AsyncMessageSystem/Public/AsyncMessageHandle.h b/Engine/Plugins/Experimental/AsyncMessageSystem/Source/AsyncMessageSystem/Public/AsyncMessageHandle.h
--- a/Engine/Plugins/Experimental/AsyncMessageSystem/Source/AsyncMessageSystem/Public/AsyncMessageHandle.h
+++ b/Engine/Plugins/Experimental/AsyncMessageSystem/Source/AsyncMessageSystem/Public/AsyncMessageHandle.h
@@ -6,13 +6,17 @@
 #include "Containers/UnrealString.h"	// For FString
 #include "HAL/Platform.h"				// For uint32 type declaration
 
+#include "AsyncMessageHandle.generated.h"
+
 /**
  * Handle used to identify a bound listener to an Async Message.
  * These handles are unique to each bound listener and created by the
  * owning FAsyncMessageSystemBase that the Message was bound to.
  */
+USTRUCT(BlueprintType)
 struct ASYNCMESSAGESYSTEM_API FAsyncMessageHandle final
 {
+	GENERATED_BODY()
 public:
 
 	friend class FAsyncMessageSystemBase;
@@ -70,11 +74,13 @@ private:
 	/**
 	 * The value of this internal handle.
 	 */
+	UPROPERTY(VisibleAnywhere, Category="Async Message")
 	uint32 InternalHandle = FAsyncMessageHandle::InvalidHandleIndex;
 
 	/**
 	 * The async message that this handle is bound to.
 	 */
+	UPROPERTY(VisibleAnywhere, Category="Async Message")
 	FAsyncMessageId BoundMessage = FAsyncMessageId::Invalid;
 
 	// TODO: Should these handles also have a pointer back to the message system that they
diff --git a/Engine/Plugins/Experimental/AsyncMessageSystem/Source/AsyncMessageSystem/Public/AsyncMessageWorldSubsystem.h b/Engine/Plugins/Experimental/AsyncMessageSystem/Source/AsyncMessageSystem/Public/AsyncMessageWorldSubsystem.h
--- a/Engine/Plugins/Experimental/AsyncMessageSystem/Source/AsyncMessageSystem/Public/AsyncMessageWorldSubsystem.h
+++ b/Engine/Plugins/Experimental/AsyncMessageSystem/Source/AsyncMessageSystem/Public/AsyncMessageWorldSubsystem.h
@@ -2,6 +2,7 @@
 
 #pragma once
 
+#include "Delegates/Delegate.h"		// For TMulticastDelegate
 #include "Engine/World.h"
 #include "Subsystems/WorldSubsystem.h"
 
@@ -45,6 +46,12 @@ public:
 	 */
 	template<class TMessageSystemType = FAsyncMessageSystemBase>
 	TSharedPtr<TMessageSystemType> GetSharedMessageSystem() const;
+
+	/**
+	 * Delegate which is broadcast when this subsystem's message system is shutdown. Can be used to clean up any bindings
+	 * to this subsystem's message system when the world goes out of scope
+	 */
+	TMulticastDelegate<void()> OnShutdownMessageSystem;
 	
 protected:
 
@@ -60,7 +67,7 @@ protected:
 	* Shutdown and cleans up this worlds message system. 
 	*/
 	void ShutdownMessageSystem();
-
+	
 	/**
 	 * Instance of the message system for this world subsystem. This shares the lifetime of this subsystem
 	 * and will be released in ShutdownMessageSystem, upon Deinitialize of this subsystem.
