Add AsyncMessage.EnableWorldSubsystem (defaults to on) and AsyncMessage.EnableWorldSubsystemForEditor (defaults to off) cvars to control when the async message world subsystem is actually created

diff --git a/Engine/Plugins/Experimental/AsyncMessageSystem/Source/AsyncMessageSystem/Private/AsyncMessageWorldSubsystem.cpp b/Engine/Plugins/Experimental/AsyncMessageSystem/Source/AsyncMessageSystem/Private/AsyncMessageWorldSubsystem.cpp
--- a/Engine/Plugins/Experimental/AsyncMessageSystem/Source/AsyncMessageSystem/Private/AsyncMessageWorldSubsystem.cpp
+++ b/Engine/Plugins/Experimental/AsyncMessageSystem/Source/AsyncMessageSystem/Private/AsyncMessageWorldSubsystem.cpp
@@ -1,9 +1,38 @@
 // Copyright Epic Games, Inc. All Rights Reserved.
 
 #include "AsyncMessageWorldSubsystem.h"
-
 #include "AsyncGameplayMessageSystem.h"
 #include "AsyncMessageSystemLogs.h"
+#include "HAL/IConsoleManager.h"
+
+static bool bAsyncMessageEnableWorldSubsystem = true;
+FAutoConsoleVariableRef CVar_AsyncMessageEnableWorldSubsystem(
+	TEXT("AsyncMessage.EnableWorldSubsystem"), bAsyncMessageEnableWorldSubsystem,
+	TEXT("When set to true, create an async message world subsystem."),
+	ECVF_Default);
+
+static bool bAsyncMessageEnableWorldSubsystemForEditor = false;
+FAutoConsoleVariableRef CVar_AsyncMessageEnableWorldSubsystemForEditor(
+	TEXT("AsyncMessage.EnableWorldSubsystemForEditor"), bAsyncMessageEnableWorldSubsystemForEditor,
+	TEXT("When set to true, create an async message world subsystem for editor worlds as well as gameplay ones."),
+	ECVF_Default);
+
+bool UAsyncMessageWorldSubsystem::ShouldCreateSubsystem(UObject* Outer) const
+{
+	if (!bAsyncMessageEnableWorldSubsystem)
+	{
+		return false;
+	}
+
+	UWorld* World = Cast<UWorld>(Outer);
+	if (!bAsyncMessageEnableWorldSubsystemForEditor && World->WorldType == EWorldType::Editor)
+	{
+		return false;
+	}
+
+	// Fall back to default behavior for editor, game, and pie worlds
+	return Super::ShouldCreateSubsystem(Outer);
+}
 
 void UAsyncMessageWorldSubsystem::Initialize(FSubsystemCollectionBase& Collection)
 {
diff --git a/Engine/Plugins/Experimental/AsyncMessageSystem/Source/AsyncMessageSystem/Public/AsyncMessageWorldSubsystem.h b/Engine/Plugins/Experimental/AsyncMessageSystem/Source/AsyncMessageSystem/Public/AsyncMessageWorldSubsystem.h
--- a/Engine/Plugins/Experimental/AsyncMessageSystem/Source/AsyncMessageSystem/Public/AsyncMessageWorldSubsystem.h
+++ b/Engine/Plugins/Experimental/AsyncMessageSystem/Source/AsyncMessageSystem/Public/AsyncMessageWorldSubsystem.h
@@ -55,6 +55,7 @@ public:
 	
 protected:
 
+	virtual bool ShouldCreateSubsystem(UObject* Outer) const override;
 	virtual void Initialize(FSubsystemCollectionBase& Collection) override;
 	virtual void Deinitialize() override;
 
