Remove some unused includes/forward decls from the message system.

#rnx
#tests non-unity compile

diff --git a/Engine/Plugins/Experimental/AsyncMessageSystem/Source/AsyncMessageSystem/Private/AsyncMessageSystemModule.cpp b/Engine/Plugins/Experimental/AsyncMessageSystem/Source/AsyncMessageSystem/Private/AsyncMessageSystemModule.cpp
--- a/Engine/Plugins/Experimental/AsyncMessageSystem/Source/AsyncMessageSystem/Private/AsyncMessageSystemModule.cpp
+++ b/Engine/Plugins/Experimental/AsyncMessageSystem/Source/AsyncMessageSystem/Private/AsyncMessageSystemModule.cpp
@@ -1,9 +1,6 @@
 // Copyright Epic Games, Inc. All Rights Reserved.
 
 #include "AsyncMessageSystemModule.h"
-#include "AsyncGameplayMessageSystem.h"
-#include "AsyncMessageSystemLogs.h"
-#include "AsyncMessageSystemBase.h"
 #include "Modules/ModuleManager.h"
 
 bool FAsyncMessageSystemModule::IsAvailable()
diff --git a/Engine/Plugins/Experimental/AsyncMessageSystem/Source/AsyncMessageSystem/Private/AsyncMessageSystemModule.h b/Engine/Plugins/Experimental/AsyncMessageSystem/Source/AsyncMessageSystem/Private/AsyncMessageSystemModule.h
--- a/Engine/Plugins/Experimental/AsyncMessageSystem/Source/AsyncMessageSystem/Private/AsyncMessageSystemModule.h
+++ b/Engine/Plugins/Experimental/AsyncMessageSystem/Source/AsyncMessageSystem/Private/AsyncMessageSystemModule.h
@@ -3,11 +3,8 @@
 #pragma once
 
 #include "Modules/ModuleInterface.h"
-#include "Templates/SharedPointer.h"
 
-class FAsyncMessageSystemBase;
-
-class FAsyncMessageSystemModule : public IModuleInterface
+class FAsyncMessageSystemModule final : public IModuleInterface
 {
 public:
 
