Create some developer settings to let you toggle the creation of async world subsystem. This way if you want to create your own subsystem, you can disable the default without having any overhead

#rnx

diff --git a/Engine/Plugins/Experimental/AsyncMessageSystem/Source/AsyncMessageSystem/AsyncMessageSystem.build.cs b/Engine/Plugins/Experimental/AsyncMessageSystem/Source/AsyncMessageSystem/AsyncMessageSystem.build.cs
--- a/Engine/Plugins/Experimental/AsyncMessageSystem/Source/AsyncMessageSystem/AsyncMessageSystem.build.cs
+++ b/Engine/Plugins/Experimental/AsyncMessageSystem/Source/AsyncMessageSystem/AsyncMessageSystem.build.cs
@@ -18,6 +18,7 @@ namespace UnrealBuildTool.Rules
 				new string[] {
 					"Core",
 					"CoreUObject",
+					"DeveloperSettings",
 					"Engine",
 					"GameplayTags",
 				}
diff --git a/Engine/Plugins/Experimental/AsyncMessageSystem/Source/AsyncMessageSystem/Private/AsyncMessageDeveloperSettings.cpp b/Engine/Plugins/Experimental/AsyncMessageSystem/Source/AsyncMessageSystem/Private/AsyncMessageDeveloperSettings.cpp
new file mode 100644
--- /dev/null
+++ b/Engine/Plugins/Experimental/AsyncMessageSystem/Source/AsyncMessageSystem/Private/AsyncMessageDeveloperSettings.cpp
@@ -0,0 +1,13 @@
+// Copyright Epic Games, Inc. All Rights Reserved.
+
+#include "AsyncMessageDeveloperSettings.h"
+
+bool UAsyncMessageDeveloperSettings::GetShouldEnableWorldSubsystem() const
+{
+	return bEnableWorldSubsystem;
+}
+
+bool UAsyncMessageDeveloperSettings::GetShouldEnableWorldSubsystemInEditor() const
+{
+	return bEnableWorldSubsystemInEditor;
+}
diff --git a/Engine/Plugins/Experimental/AsyncMessageSystem/Source/AsyncMessageSystem/Private/AsyncMessageDeveloperSettings.h b/Engine/Plugins/Experimental/AsyncMessageSystem/Source/AsyncMessageSystem/Private/AsyncMessageDeveloperSettings.h
new file mode 100644
--- /dev/null
+++ b/Engine/Plugins/Experimental/AsyncMessageSystem/Source/AsyncMessageSystem/Private/AsyncMessageDeveloperSettings.h
@@ -0,0 +1,31 @@
+// Copyright Epic Games, Inc. All Rights Reserved.
+
+#pragma once
+
+#include "Engine/DeveloperSettings.h"
+
+#include "AsyncMessageDeveloperSettings.generated.h"
+
+UCLASS(Config=Game)
+class UAsyncMessageDeveloperSettings final : public UDeveloperSettings
+{
+	GENERATED_BODY()
+
+public:
+	
+	bool GetShouldEnableWorldSubsystem() const;
+	bool GetShouldEnableWorldSubsystemInEditor() const;
+
+private:
+	/**
+	 * If true, then the async message world subsystem will be enabled.
+	 */
+	UPROPERTY(EditAnywhere, Config, Category = "Async Messages")
+	bool bEnableWorldSubsystem = true;
+
+	/**
+	 * If true, then the async message subsystem will be created for editor worlds.
+	 */
+	UPROPERTY(EditAnywhere, Config, Category = "Async Messages")
+	bool bEnableWorldSubsystemInEditor = false;
+};
\ No newline at end of file
diff --git a/Engine/Plugins/Experimental/AsyncMessageSystem/Source/AsyncMessageSystem/Private/AsyncMessageWorldSubsystem.cpp b/Engine/Plugins/Experimental/AsyncMessageSystem/Source/AsyncMessageSystem/Private/AsyncMessageWorldSubsystem.cpp
--- a/Engine/Plugins/Experimental/AsyncMessageSystem/Source/AsyncMessageSystem/Private/AsyncMessageWorldSubsystem.cpp
+++ b/Engine/Plugins/Experimental/AsyncMessageSystem/Source/AsyncMessageSystem/Private/AsyncMessageWorldSubsystem.cpp
@@ -3,19 +3,7 @@
 #include "AsyncMessageWorldSubsystem.h"
 #include "AsyncGameplayMessageSystem.h"
 #include "AsyncMessageSystemLogs.h"
-#include "HAL/IConsoleManager.h"
-
-static bool bAsyncMessageEnableWorldSubsystem = true;
-FAutoConsoleVariableRef CVar_AsyncMessageEnableWorldSubsystem(
-	TEXT("AsyncMessage.EnableWorldSubsystem"), bAsyncMessageEnableWorldSubsystem,
-	TEXT("When set to true, create an async message world subsystem."),
-	ECVF_Default);
-
-static bool bAsyncMessageEnableWorldSubsystemForEditor = false;
-FAutoConsoleVariableRef CVar_AsyncMessageEnableWorldSubsystemForEditor(
-	TEXT("AsyncMessage.EnableWorldSubsystemForEditor"), bAsyncMessageEnableWorldSubsystemForEditor,
-	TEXT("When set to true, create an async message world subsystem for editor worlds as well as gameplay ones."),
-	ECVF_Default);
+#include "AsyncMessageDeveloperSettings.h"
 
 void UAsyncMessageWorldSubsystem::AddReferencedObjects(UObject* InThis, FReferenceCollector& Collector)
 {
@@ -29,15 +17,29 @@ void UAsyncMessageWorldSubsystem::AddReferencedObjects(UObject* InThis, FReferen
 	}
 }
 
-bool UAsyncMessageWorldSubsystem::ShouldCreateSubsystem(UObject* Outer) const
+namespace UE::Private
 {
-	if (!bAsyncMessageEnableWorldSubsystem)
+	static bool ShouldCreateWorldSubsystem(const UWorld* ForWorld)
 	{
-		return false;
+		if (!ForWorld)
+		{
+			return false;
+		}
+		const UAsyncMessageDeveloperSettings* Settings = GetDefault<UAsyncMessageDeveloperSettings>();
+	
+		// Let you disable for editor worlds
+		if (ForWorld->WorldType == EWorldType::Editor)
+		{
+			return Settings->GetShouldEnableWorldSubsystemInEditor();
+		}
+	
+		return Settings->GetShouldEnableWorldSubsystem();
 	}
+};
 
-	UWorld* World = Cast<UWorld>(Outer);
-	if (!bAsyncMessageEnableWorldSubsystemForEditor && World->WorldType == EWorldType::Editor)
+bool UAsyncMessageWorldSubsystem::ShouldCreateSubsystem(UObject* Outer) const
+{
+	if (!UE::Private::ShouldCreateWorldSubsystem(Cast<UWorld>(Outer)))
 	{
 		return false;
 	}
