Create a way to record the callstack and trigger a breakpoint of messages being broadcast to more easily debug the message system

You can now enable logging of messages when they are queued and their callstack if the ENABLE_ASYNC_MESSAGES_DEBUG is set to 1.

By default, ENABLE_ASYNC_MESSAGES_DEBUG is set to 1 in editor and non-shipping builds. You can also override it and set it via your game's .build.cs file.

When enabled, you can set which Message Id's you would like to log out the callstack for, or enable logging of all messages.

#rnx

diff --git a/Engine/Plugins/Experimental/AsyncMessageSystem/Source/AsyncMessageSystem/Private/AsyncMessageDeveloperSettings.cpp b/Engine/Plugins/Experimental/AsyncMessageSystem/Source/AsyncMessageSystem/Private/AsyncMessageDeveloperSettings.cpp
--- a/Engine/Plugins/Experimental/AsyncMessageSystem/Source/AsyncMessageSystem/Private/AsyncMessageDeveloperSettings.cpp
+++ b/Engine/Plugins/Experimental/AsyncMessageSystem/Source/AsyncMessageSystem/Private/AsyncMessageDeveloperSettings.cpp
@@ -11,3 +11,28 @@ bool UAsyncMessageDeveloperSettings::GetShouldEnableWorldSubsystemInEditor() con
 {
 	return bEnableWorldSubsystemInEditor;
 }
+
+bool UAsyncMessageDeveloperSettings::ShouldDebugMessageOnQueue(const FAsyncMessageId& Message) const
+{
+	if (!bMessageQueueDebugEnabled)
+	{
+		return false;
+	}
+
+	if (bEnabledDebuggingForAllQueuedMessages)
+	{
+		return true;
+	}
+
+	return EnabledDebugMessages.Contains(Message);
+}
+
+bool UAsyncMessageDeveloperSettings::ShouldTriggerBreakPointOnMessageQueue() const
+{
+	return bTriggerDebugBreakpointWhenMessageQueued;
+}
+
+bool UAsyncMessageDeveloperSettings::ShouldPrintScriptCallstackOnMessageQueue() const
+{
+	return bPrintScriptCallstackWhenMessageQueued;
+}
diff --git a/Engine/Plugins/Experimental/AsyncMessageSystem/Source/AsyncMessageSystem/Private/AsyncMessageDeveloperSettings.h b/Engine/Plugins/Experimental/AsyncMessageSystem/Source/AsyncMessageSystem/Private/AsyncMessageDeveloperSettings.h
--- a/Engine/Plugins/Experimental/AsyncMessageSystem/Source/AsyncMessageSystem/Private/AsyncMessageDeveloperSettings.h
+++ b/Engine/Plugins/Experimental/AsyncMessageSystem/Source/AsyncMessageSystem/Private/AsyncMessageDeveloperSettings.h
@@ -2,6 +2,7 @@
 
 #pragma once
 
+#include "AsyncMessageId.h"
 #include "Engine/DeveloperSettings.h"
 
 #include "AsyncMessageDeveloperSettings.generated.h"
@@ -13,19 +14,91 @@ class UAsyncMessageDeveloperSettings final : public UDeveloperSettings
 
 public:
 	
-	bool GetShouldEnableWorldSubsystem() const;
-	bool GetShouldEnableWorldSubsystemInEditor() const;
+	[[nodiscard]] bool GetShouldEnableWorldSubsystem() const;
+	[[nodiscard]] bool GetShouldEnableWorldSubsystemInEditor() const;
+
+	/**
+	 * @param Message The message to check if it should have its stack trace dumped to the log 
+	 * @return True if this message should dump its current stack trace to the log
+	 */
+	[[nodiscard]] bool ShouldDebugMessageOnQueue(const FAsyncMessageId& Message) const;
+
+	/**
+	 * @return True if a breakpoint should be triggered when a message id with debugging enabled is queued for broadcast. 
+	 */
+	[[nodiscard]] bool ShouldTriggerBreakPointOnMessageQueue() const;
+
+	/**
+	 * @return True if the Blueprint VM callstack should be printed when a message with debugging enabled is queued for broadcast.
+	 */
+	[[nodiscard]] bool ShouldPrintScriptCallstackOnMessageQueue() const;
 
 private:
+	
+	/**
+	 * Enabled the conditional logging of stack traces and C++ breakpoints being triggered when certain
+	 * messages are queued for broadcasting to the message system.
+	 *
+	 * This can make it easier to debug where messages are coming from, since the processing of
+	 * messages is all deferred and it can be difficult to track down what is queuing a message in
+	 * cooked/optimized builds.
+	 *
+	 * This will only work if the ENABLE_ASYNC_MESSAGES_DEBUG pre-processor definition is non-zero.
+	 * By default, ENABLE_ASYNC_MESSAGES_DEBUG is "1" for any non-shipping or editor build (WITH_EDITOR || !UE_BUILD_SHIPPING).
+	 * 
+	 * If you wish to override this behavior, then you can add a definition to your game's .build.cs file:
+	 *		PublicDefinitions.Add("ENABLE_ASYNC_MESSAGES_DEBUG=1");
+	 *
+	 * @see AsyncMessageSystemBase.h
+	 * @see FAsyncMessageSystemBase::QueueMessageForBroadcast
+	 */
+	UPROPERTY(EditDefaultsOnly, Config, Category = "Async Messages|Debug")
+	bool bMessageQueueDebugEnabled = false;
+
+	/**
+	 * If true, then a native C++ breakpoint will be triggered when when certain
+	 * messages are queued for broadcasting to the message system.
+	 * 
+	 * @see FAsyncMessageSystemBase::QueueMessageForBroadcast
+	 */
+	UPROPERTY(EditDefaultsOnly, Config, Category = "Async Messages|Debug", meta=(EditCondition="bMessageQueueDebugEnabled"))
+	bool bTriggerDebugBreakpointWhenMessageQueued = false;
+
+	/**
+	 * If true, then the script callstack will also be printed when certain
+	 * messages are queued for broadcasting to the message system.
+	 *
+	 * @see FAsyncMessageSystemBase::QueueMessageForBroadcast
+	 */
+	UPROPERTY(EditDefaultsOnly, Config, Category = "Async Messages|Debug", meta=(EditCondition="bMessageQueueDebugEnabled"))
+	bool bPrintScriptCallstackWhenMessageQueued = false;
+
+	/**
+	 * If true, then ALL messages queued for broadcast will have their debug information processed.
+	 * 
+	 * Note: This will likely have a large performance impact.
+	 *
+	 * If false, then only messages in the "EnabledDebugMessages" array will be debugged when queued.
+	 */
+	UPROPERTY(EditDefaultsOnly, Config, Category = "Async Messages|Debug", meta=(EditCondition="bMessageQueueDebugEnabled"))
+	bool bEnabledDebuggingForAllQueuedMessages = false;
+	
+	/**
+	* Array of Async Message Id's which you would like to enable for debugging 
+	* when they are queued for broadcasting to the message system.
+	*/
+	UPROPERTY(EditDefaultsOnly, Config, Category = "Async Messages|Debug", meta=(EditCondition="bMessageQueueDebugEnabled && !bEnabledDebuggingForAllQueuedMessages"))
+	TArray<FAsyncMessageId> EnabledDebugMessages;
+	
 	/**
 	 * If true, then the async message world subsystem will be enabled.
 	 */
-	UPROPERTY(EditAnywhere, Config, Category = "Async Messages")
+	UPROPERTY(EditDefaultsOnly, Config, Category = "Async Messages")
 	bool bEnableWorldSubsystem = true;
 
 	/**
 	 * If true, then the async message subsystem will be created for editor worlds.
 	 */
-	UPROPERTY(EditAnywhere, Config, Category = "Async Messages")
+	UPROPERTY(EditDefaultsOnly, Config, Category = "Async Messages")
 	bool bEnableWorldSubsystemInEditor = false;
 };
\ No newline at end of file
diff --git a/Engine/Plugins/Experimental/AsyncMessageSystem/Source/AsyncMessageSystem/Private/AsyncMessageSystemBase.cpp b/Engine/Plugins/Experimental/AsyncMessageSystem/Source/AsyncMessageSystem/Private/AsyncMessageSystemBase.cpp
--- a/Engine/Plugins/Experimental/AsyncMessageSystem/Source/AsyncMessageSystem/Private/AsyncMessageSystemBase.cpp
+++ b/Engine/Plugins/Experimental/AsyncMessageSystem/Source/AsyncMessageSystem/Private/AsyncMessageSystemBase.cpp
@@ -2,11 +2,16 @@
 
 #include "AsyncMessageSystemBase.h"
 
+#include "AsyncMessageDeveloperSettings.h"
 #include "AsyncMessageSystemLogs.h"
 #include "CoreGlobals.h"
 #include "Misc/App.h"
 #include "NativeGameplayTags.h"
 
+#if ENABLE_ASYNC_MESSAGES_DEBUG
+#include "Misc/AssertionMacros.h"
+#endif
+
 FAsyncMessageSystemBase::~FAsyncMessageSystemBase()
 {
 	// Make sure that you call the "Shutdown" function on your message system to allow it to clean up before destruction!
@@ -32,7 +37,42 @@ void FAsyncMessageSystemBase::UnbindListener(const FAsyncMessageHandle& HandleTo
 
 bool FAsyncMessageSystemBase::QueueMessageForBroadcast(const FAsyncMessageId MessageId, FConstStructView PayloadData)
 {
+#if ENABLE_ASYNC_MESSAGES_DEBUG
+	const UAsyncMessageDeveloperSettings* Settings = GetDefault<UAsyncMessageDeveloperSettings>();
+	
+	// If this message is configured to have its stack trace dumped when queued, then do so
+	if (Settings && Settings->ShouldDebugMessageOnQueue(MessageId))
+	{
+		const uint32 ThreadQueuedFrom = FPlatformTLS::GetCurrentThreadId();
+		const FString Heading = FString::Printf(TEXT("=== Message '%s' queued from thread %u ==="), *MessageId.ToString(), ThreadQueuedFrom);
+		
+		FDebug::DumpStackTraceToLog(*Heading, ELogVerbosity::Display);
+
+		// Print the BP callstack if desired
+		if (Settings->ShouldPrintScriptCallstackOnMessageQueue())
+		{
+			PrintScriptCallstack();	
+		}
+		
+		// Note: In case you forgot to enable the printing of the Blueprint VM callstack, you can use these commands
+		// in the "immediate" window to print it instead:
+		// 
+		// In editor builds:
+		//		{,,UnrealEditor-Core}::PrintScriptCallstack()
+		// In a game/client build:
+		//		::PrintScriptCallstack()
+
+		// Also allow for users to easily configure breakpointers when these messages are queued, making it easy
+		// to track down where and why the messages are being broadcast.
+		if (Settings->ShouldTriggerBreakPointOnMessageQueue())
+		{
+			UE_DEBUG_BREAK();	
+		}
+	}
+#endif	// ENABLE_ASYNC_MESSAGES_DEBUG
+
 	TRACE_CPUPROFILER_EVENT_SCOPE(FAsyncMessageSystemBase::QueueMessageForBroadcast);
+	
 	return QueueMessageForBroadcast_Impl(MessageId, MoveTemp(PayloadData));
 }
 
diff --git a/Engine/Plugins/Experimental/AsyncMessageSystem/Source/AsyncMessageSystem/Public/AsyncMessageSystemBase.h b/Engine/Plugins/Experimental/AsyncMessageSystem/Source/AsyncMessageSystem/Public/AsyncMessageSystemBase.h
--- a/Engine/Plugins/Experimental/AsyncMessageSystem/Source/AsyncMessageSystem/Public/AsyncMessageSystemBase.h
+++ b/Engine/Plugins/Experimental/AsyncMessageSystem/Source/AsyncMessageSystem/Public/AsyncMessageSystemBase.h
@@ -415,3 +415,17 @@ FAsyncMessageHandle FAsyncMessageSystemBase::BindListener(
 
 	return bWasSuccessfullyBound ? NewListenerHandle : FAsyncMessageHandle::Invalid;
 }
+
+// A pre-processor define which will allow you to log the callstack of messages when they are queued.
+// This can be helpful when debugging a build to see what messages are being queued and from where,
+// since it can be hard to determine that information from the listener.
+//
+// By default, this will be enabled in the editor and non-shipping bulids.
+//
+// If you would like to enable this in a shipping build, then add a define to your game's .build.cs file:
+//		PublicDefinitions.Add("ENABLE_ASYNC_MESSAGES_DEBUG=1");
+//
+// It is not recommended that you log out on every message, because it will add significant performance cost.
+#ifndef ENABLE_ASYNC_MESSAGES_DEBUG
+	#define ENABLE_ASYNC_MESSAGES_DEBUG WITH_EDITOR || !UE_BUILD_SHIPPING
+#endif
\ No newline at end of file
