store the callstack of when a message is queued on the message itself

#rnx

diff --git a/Engine/Plugins/Experimental/AsyncMessageSystem/Source/AsyncMessageSystem/Private/AsyncMessage.cpp b/Engine/Plugins/Experimental/AsyncMessageSystem/Source/AsyncMessageSystem/Private/AsyncMessage.cpp
--- a/Engine/Plugins/Experimental/AsyncMessageSystem/Source/AsyncMessageSystem/Private/AsyncMessage.cpp
+++ b/Engine/Plugins/Experimental/AsyncMessageSystem/Source/AsyncMessageSystem/Private/AsyncMessage.cpp
@@ -2,6 +2,10 @@
 
 #include "AsyncMessage.h"
 
+#if ENABLE_ASYNC_MESSAGES_DEBUG
+#include "HAL/PlatformStackWalk.h"
+#endif
+
 FAsyncMessage::FAsyncMessage(
 	const FAsyncMessageId& InMessageId,
 	const FAsyncMessageId& InMessageSourceId,
@@ -70,3 +74,59 @@ FConstStructView FAsyncMessage::GetPayloadView() const
 {
 	return FConstStructView(PayloadCopy);
 }
+
+#if ENABLE_ASYNC_MESSAGES_DEBUG
+
+const FString& FAsyncMessage::FMessageDebugData::GetOrCreateNativeCallstackAsString() const
+{
+	if (NativeCallstackAsString.Len() > 0)
+	{
+		return NativeCallstackAsString;
+	}
+
+	// Decode the callstack and cache it in a string
+	for (int32 i = 0; i < NativeCallstack.Num(); i++)
+	{
+		FProgramCounterSymbolInfo LineInfo;
+		FPlatformStackWalk::ProgramCounterToSymbolInfo(NativeCallstack[i], LineInfo);
+
+		const ANSICHAR* FilenameWithoutPath = nullptr;
+		if ((FilenameWithoutPath = FCStringAnsi::Strrchr(LineInfo.Filename, '/')) == nullptr)
+		{
+			if ((FilenameWithoutPath = FCStringAnsi::Strrchr(LineInfo.Filename, '\\')) == nullptr)
+			{
+				FilenameWithoutPath = LineInfo.Filename;
+			}
+		}
+
+		ANSICHAR LineString[MAX_SPRINTF];
+		FCStringAnsi::Sprintf(LineString, "%-64s (%s:%i)\n", LineInfo.FunctionName, FilenameWithoutPath+1, LineInfo.LineNumber);
+
+		NativeCallstackAsString.Append(LineString);
+	}
+
+	return NativeCallstackAsString;
+}
+
+const FString& FAsyncMessage::GetNativeCallstack() const
+{
+	static const FString Unknown = TEXT("Unknown");
+	return DebugData ? DebugData->GetOrCreateNativeCallstackAsString() : Unknown;
+}
+
+const FString& FAsyncMessage::GetBlueprintScriptCallstack() const
+{
+	static const FString Unknown = TEXT("Unknown");
+	return DebugData ? DebugData->BlueprintScriptCallstack : Unknown;
+}
+
+const uint32 FAsyncMessage::GetDebugMessageId() const
+{
+	return DebugData ? DebugData->MessageId : 0u;
+}
+
+void FAsyncMessage::SetDebugData(FMessageDebugData* InData)
+{
+	DebugData = InData;
+}
+#endif	// #if ENABLE_ASYNC_MESSAGES_DEBUG
diff --git a/Engine/Plugins/Experimental/AsyncMessageSystem/Source/AsyncMessageSystem/Private/AsyncMessageBlueprintLibrary.cpp b/Engine/Plugins/Experimental/AsyncMessageSystem/Source/AsyncMessageSystem/Private/AsyncMessageBlueprintLibrary.cpp
--- a/Engine/Plugins/Experimental/AsyncMessageSystem/Source/AsyncMessageSystem/Private/AsyncMessageBlueprintLibrary.cpp
+++ b/Engine/Plugins/Experimental/AsyncMessageSystem/Source/AsyncMessageSystem/Private/AsyncMessageBlueprintLibrary.cpp
@@ -38,3 +38,21 @@ FString UAsyncMessageSystemBlueprintLibrary::Conv_AsyncMessageIdToString(const F
 {
 	return MessageId.ToString();
 }
+
+FString UAsyncMessageSystemBlueprintLibrary::GetMessageNativeQueueCallstack(const FAsyncMessage& Message)
+{
+#if ENABLE_ASYNC_MESSAGES_DEBUG
+	return Message.GetNativeCallstack();
+#else
+	return TEXT("Unknown: ENABLE_ASYNC_MESSAGES_DEBUG is disabled in this build configuration");
+#endif	// #if ENABLE_ASYNC_MESSAGES_DEBUG
+}
+
+FString UAsyncMessageSystemBlueprintLibrary::GetMessageBlueprintScriptCallstack(const FAsyncMessage& Message)
+{
+#if ENABLE_ASYNC_MESSAGES_DEBUG
+	return Message.GetBlueprintScriptCallstack();
+#else
+	return TEXT("Unknown: ENABLE_ASYNC_MESSAGES_DEBUG is disabled in this build configuration");
+#endif	// #if ENABLE_ASYNC_MESSAGES_DEBUG
+}
\ No newline at end of file
diff --git a/Engine/Plugins/Experimental/AsyncMessageSystem/Source/AsyncMessageSystem/Private/AsyncMessageBlueprintLibrary.h b/Engine/Plugins/Experimental/AsyncMessageSystem/Source/AsyncMessageSystem/Private/AsyncMessageBlueprintLibrary.h
--- a/Engine/Plugins/Experimental/AsyncMessageSystem/Source/AsyncMessageSystem/Private/AsyncMessageBlueprintLibrary.h
+++ b/Engine/Plugins/Experimental/AsyncMessageSystem/Source/AsyncMessageSystem/Private/AsyncMessageBlueprintLibrary.h
@@ -42,4 +42,27 @@ public:
 	 */
 	UFUNCTION(BlueprintPure, Category = "Utilities|String", meta=(DisplayName = "To String (Async Message Id)", CompactNodeTitle = "->", BlueprintAutocast))
 	static FString Conv_AsyncMessageIdToString(const FAsyncMessageId& MessageId);
+
+	/**
+	 * Gets the given message's callstack of when and where it was queued from native C++ code.
+	 *
+	 * Note: bShouldRecordQueueCallstackOnMessages must be enabled in the project settings for this to have accurate data
+	 * 
+	 * @param Message The async message to get the callstack of to determine when and where it was queued.
+	 * @return The callstack of where the message was queued from
+	 */
+	UFUNCTION(BlueprintPure, Category = "Async Messages|Debug")
+	static FString GetMessageNativeQueueCallstack(const FAsyncMessage& Message);
+
+	/**
+	 * Gets the given message's callstack of when and where it was queued from in blueprints/script
+	 *
+	 * Note: bShouldRecordQueueCallstackOnMessages must be enabled in the project settings for this to have accurate data
+	 * 
+	 * @param Message The async message to get the callstack of to determine when and where it was queued.
+	 * @return The script callstack of where the message was queued from
+	 */
+	UFUNCTION(BlueprintPure, Category = "Async Messages|Debug")
+	static FString GetMessageBlueprintScriptCallstack(const FAsyncMessage& Message);
+	
 };
diff --git a/Engine/Plugins/Experimental/AsyncMessageSystem/Source/AsyncMessageSystem/Private/AsyncMessageDeveloperSettings.cpp b/Engine/Plugins/Experimental/AsyncMessageSystem/Source/AsyncMessageSystem/Private/AsyncMessageDeveloperSettings.cpp
--- a/Engine/Plugins/Experimental/AsyncMessageSystem/Source/AsyncMessageSystem/Private/AsyncMessageDeveloperSettings.cpp
+++ b/Engine/Plugins/Experimental/AsyncMessageSystem/Source/AsyncMessageSystem/Private/AsyncMessageDeveloperSettings.cpp
@@ -36,3 +36,8 @@ bool UAsyncMessageDeveloperSettings::ShouldPrintScriptCallstackOnMessageQueue()
 {
 	return bPrintScriptCallstackWhenMessageQueued;
 }
+
+bool UAsyncMessageDeveloperSettings::ShouldRecordQueueCallstackOnMessages() const
+{
+	return bShouldRecordQueueCallstackOnMessages;
+}
diff --git a/Engine/Plugins/Experimental/AsyncMessageSystem/Source/AsyncMessageSystem/Private/AsyncMessageDeveloperSettings.h b/Engine/Plugins/Experimental/AsyncMessageSystem/Source/AsyncMessageSystem/Private/AsyncMessageDeveloperSettings.h
--- a/Engine/Plugins/Experimental/AsyncMessageSystem/Source/AsyncMessageSystem/Private/AsyncMessageDeveloperSettings.h
+++ b/Engine/Plugins/Experimental/AsyncMessageSystem/Source/AsyncMessageSystem/Private/AsyncMessageDeveloperSettings.h
@@ -33,6 +33,13 @@ public:
 	 */
 	[[nodiscard]] bool ShouldPrintScriptCallstackOnMessageQueue() const;
 
+	/**
+	 * @return True if the callstack should be recorded as a property on the Async Message itself. This
+	 * can make debugging where a message was queued from signifigantly easier if you are starting from the
+	 * listener. 
+	 */
+	[[nodiscard]] bool ShouldRecordQueueCallstackOnMessages() const;
+
 private:
 	
 	/**
@@ -73,6 +80,16 @@ private:
 	UPROPERTY(EditDefaultsOnly, Config, Category = "Async Messages|Debug", meta=(EditCondition="bMessageQueueDebugEnabled"))
 	bool bPrintScriptCallstackWhenMessageQueued = false;
 
+	/**
+	 * If true, then the callstack at the time of a message being queued will be recorded and stored on the
+	 * FAsyncMessage instance itself. This can make debugging listeners significantly easier if you
+	 * need to get an idea of where a message is coming from.
+	 *
+	 * @see FAsyncMessageSystemBase::QueueMessageForBroadcast_Impl
+	 */
+	UPROPERTY(EditDefaultsOnly, Config, Category = "Async Messages|Debug")
+	bool bShouldRecordQueueCallstackOnMessages = false;
+
 	/**
 	 * If true, then ALL messages queued for broadcast will have their debug information processed.
 	 * 
diff --git a/Engine/Plugins/Experimental/AsyncMessageSystem/Source/AsyncMessageSystem/Private/AsyncMessageStore.cpp b/Engine/Plugins/Experimental/AsyncMessageSystem/Source/AsyncMessageSystem/Private/AsyncMessageStore.cpp
--- a/Engine/Plugins/Experimental/AsyncMessageSystem/Source/AsyncMessageSystem/Private/AsyncMessageStore.cpp
+++ b/Engine/Plugins/Experimental/AsyncMessageSystem/Source/AsyncMessageSystem/Private/AsyncMessageStore.cpp
@@ -2,6 +2,36 @@
 
 #include "AsyncMessageStore.h"
 
+#if ENABLE_ASYNC_MESSAGES_DEBUG
+
+#include "AsyncMessageDeveloperSettings.h"
+#include "Stats/StatsMisc.h"
+
+uint32 FAsyncMessageStore::CreateDebugMessageData()
+{
+	const uint32 DebugDataId = NextDebugMessageId++;
+	TTuple<uint32, FAsyncMessage::FMessageDebugData>& DebugData = MessageDebugCallstacks.Add(DebugDataId, TTuple<uint32, FAsyncMessage::FMessageDebugData>{ });
+
+	DebugData.Value.MessageId = DebugDataId;
+	
+	// Stack trace. Do this fast and don't bother getting the symbols (tanks performance)
+	constexpr int32 StackTraceDepth = 32;
+	uint64 StackTrace[StackTraceDepth];
+	const uint32 StackDepth = FPlatformStackWalk::CaptureStackBackTrace(StackTrace, StackTraceDepth);
+	const uint32 NumStackItemsToSkip = 6;
+	
+	if (StackDepth > NumStackItemsToSkip)
+	{
+		DebugData.Value.NativeCallstack = TArray<uint64>(&StackTrace[NumStackItemsToSkip], StackDepth - NumStackItemsToSkip);
+	}
+
+	// Also record the BP script callstack in case this was queued from a blueprint/script call
+	DebugData.Value.BlueprintScriptCallstack = FFrame::GetScriptCallstack();
+	
+	return DebugDataId;
+}
+#endif	// #if ENABLE_ASYNC_MESSAGES_DEBUG
+
 uint32 FAsyncMessageStore::EnqueueMessage(FAsyncMessage&& MessageToQueue, OUT TArray<FAsyncMessageBindingOptions>& BindingTypesAddedTo)
 {
 	TRACE_CPUPROFILER_EVENT_SCOPE(FAsyncMessageStore::EnqueueMessage);
@@ -13,6 +43,11 @@ uint32 FAsyncMessageStore::EnqueueMessage(FAsyncMessage&& MessageToQueue, OUT TA
 	
 	uint32 NumQueuesAddedTo = 0u;
 
+#if ENABLE_ASYNC_MESSAGES_DEBUG
+	const bool bRecordQueueCallstack = GetDefault<UAsyncMessageDeveloperSettings>()->ShouldRecordQueueCallstackOnMessages();
+	uint32 DebugDataId = 0u;
+#endif	// #if ENABLE_ASYNC_MESSAGES_DEBUG
+
 	// Check if there are any listeners to this event or any of its parents. If there are, add it to their binding queue
 	FAsyncMessageId::WalkMessageHierarchy(MessageToQueue.GetMessageId(), [&](const FAsyncMessageId& CurrentMessageId)
 	{
@@ -22,10 +57,27 @@ uint32 FAsyncMessageStore::EnqueueMessage(FAsyncMessage&& MessageToQueue, OUT TA
 			// Only add this message to a queue if we know that it has a listener for this particular message.
 			if (Pair.Value.MessagesWithTheseBindings.Contains(CurrentMessageId) && !BindingTypesAddedTo.Contains(Pair.Key))
 			{
-				Pair.Value.MessageQueue->Add(MessageToQueue);
+				const int32 AddedMessageIdx = Pair.Value.MessageQueue->Add(MessageToQueue);
 				++NumQueuesAddedTo;
 			
 				BindingTypesAddedTo.Emplace(Pair.Key);
+
+#if ENABLE_ASYNC_MESSAGES_DEBUG
+				if (bRecordQueueCallstack)
+				{
+					// Find or create debug data here and set a pointer to it on each message
+					if (DebugDataId == 0u)
+					{
+						DebugDataId = CreateDebugMessageData();
+					}
+
+					TTuple<uint32, FAsyncMessage::FMessageDebugData>& DebugData = MessageDebugCallstacks[DebugDataId];
+					++DebugData.Key;
+					check(DebugData.Key == NumQueuesAddedTo);
+					
+					(*Pair.Value.MessageQueue)[AddedMessageIdx].SetDebugData(&DebugData.Value);
+				}
+#endif	// #if ENABLE_ASYNC_MESSAGES_DEBUG
 			}
 		}	
 	});
@@ -43,7 +95,28 @@ TOptional<FAsyncMessage> FAsyncMessageStore::GetNextMessageForBindingOption(cons
 		// If we can, then pop an element off the message queue
 		if (!QueueContainer->MessageQueue->IsEmpty())
 		{
+#if ENABLE_ASYNC_MESSAGES_DEBUG
+			TOptional<FAsyncMessage> OutVal{ QueueContainer->MessageQueue->Pop(EAllowShrinking::No) };
+
+			// Remove the debug data from the map if the number of queues that the message was initially added to
+			// have all been processed
+			if (OutVal.IsSet())
+			{
+				const uint32 DebugMessageId = (*OutVal).GetDebugMessageId();
+				if (TTuple<uint32, FAsyncMessage::FMessageDebugData>* Data = MessageDebugCallstacks.Find(DebugMessageId))
+				{
+					--Data->Key;
+					if (Data->Key == 0)
+					{
+						MessageDebugCallstacks.Remove(DebugMessageId);
+					}
+				}
+			}
+			
+			return OutVal;
+#else
 			return QueueContainer->MessageQueue->Pop(EAllowShrinking::No);
+#endif
 		}
 	}
 	
diff --git a/Engine/Plugins/Experimental/AsyncMessageSystem/Source/AsyncMessageSystem/Public/AsyncMessage.h b/Engine/Plugins/Experimental/AsyncMessageSystem/Source/AsyncMessageSystem/Public/AsyncMessage.h
--- a/Engine/Plugins/Experimental/AsyncMessageSystem/Source/AsyncMessageSystem/Public/AsyncMessage.h
+++ b/Engine/Plugins/Experimental/AsyncMessageSystem/Source/AsyncMessageSystem/Public/AsyncMessage.h
@@ -8,6 +8,20 @@
 
 #include "AsyncMessage.generated.h"
 
+// A pre-processor define which will allow you to log the callstack of messages when they are queued.
+// This can be helpful when debugging a build to see what messages are being queued and from where,
+// since it can be hard to determine that information from the listener.
+//
+// By default, this will be enabled in the editor and non-shipping bulids.
+//
+// If you would like to enable this in a shipping build, then add a define to your game's .build.cs file:
+//		PublicDefinitions.Add("ENABLE_ASYNC_MESSAGES_DEBUG=1");
+//
+// It is not recommended that you log out on every message, because it will add significant performance cost.
+#ifndef ENABLE_ASYNC_MESSAGES_DEBUG
+	#define ENABLE_ASYNC_MESSAGES_DEBUG WITH_EDITOR || !UE_BUILD_SHIPPING
+#endif
+
 /**
  * A single async message which has been queued for broadcasting, and listeners are
  * interested in receiving the data from.
@@ -159,5 +173,56 @@ private:
 	FInstancedStruct PayloadCopy;
 
 	// TODO: We could perhaps define a "lifetime" concept on these messages, similar to how AnimNext's
-	// TraitEvents (TraitEvent.h) work, as well as consumption perhaps?	
+	// TraitEvents (TraitEvent.h) work, as well as consumption perhaps?
+
+	/////////////////////////////////////////////////////////////////////
+	// Debugging helpers
+	
+#if ENABLE_ASYNC_MESSAGES_DEBUG
+public:
+	
+	struct FMessageDebugData
+	{
+		/**
+		 * The native stack back trace at the time of this message being queued.
+		 * Populated by FPlatformStackWalk::CaptureStackBackTrace in FAsyncMessageSystemBase::QueueMessageForBroadcast_Impl.
+		 */
+		TArray<uint64> NativeCallstack;
+	
+		/**
+		 * A debugging nicety. This is a string which holds the callstack from where this message was queued from.
+		 * Mutable so that we can cache the value in GetOrCreateNativeCallstackAsString.
+		 */
+		mutable FString NativeCallstackAsString;
+
+		/**
+		 * A debugging nicety. This is a string which holds the blueprint/script callstack from where this message was queued from.
+		 */
+		FString BlueprintScriptCallstack;
+
+		/**
+		* Unique int id assigned at the time of creation. Each message gets a unique ID which can then be used to look up
+		* the debug data on the message store. 
+		*/
+		uint32 MessageId = 0u;
+
+		/**
+		 * Generates and caches the FString version of the platform callstack
+		 * @return the Cachced string which represents the callstack.
+		 */
+		const FString& GetOrCreateNativeCallstackAsString() const;
+	};
+	
+	const FString& GetNativeCallstack() const;
+
+	const FString& GetBlueprintScriptCallstack() const;
+
+	const uint32 GetDebugMessageId() const;
+
+	void SetDebugData(FMessageDebugData* InData);
+	
+private:
+
+	FMessageDebugData* DebugData = nullptr;
+#endif	// #if ENABLE_ASYNC_MESSAGES_DEBUG
 };
\ No newline at end of file
diff --git a/Engine/Plugins/Experimental/AsyncMessageSystem/Source/AsyncMessageSystem/Public/AsyncMessageStore.h b/Engine/Plugins/Experimental/AsyncMessageSystem/Source/AsyncMessageSystem/Public/AsyncMessageStore.h
--- a/Engine/Plugins/Experimental/AsyncMessageSystem/Source/AsyncMessageSystem/Public/AsyncMessageStore.h
+++ b/Engine/Plugins/Experimental/AsyncMessageSystem/Source/AsyncMessageSystem/Public/AsyncMessageStore.h
@@ -82,4 +82,27 @@ protected:
 	 * the listeners of the messages to control _when_ they receive an message during the frame. 
 	 */
 	TMap<FAsyncMessageBindingOptions, FBindingOptionsMessageQueue> MessageBindingQueues;
+
+#if ENABLE_ASYNC_MESSAGES_DEBUG
+	/**
+	 * Unique ID assigned to a message if we are recording its queue callstack.
+	 * 0 would be an invalid message id.
+	 */
+	std::atomic<uint32> NextDebugMessageId = 1u;
+
+	/**
+	 * Creates and adds a debug messsage to the MessageDebugCallstacks map.
+	 * Records the current native and BP/script callstacks 
+	 * @return Debug message Id which can be used to look it up in the MessageDebugCallstacks map.
+	 */
+	uint32 CreateDebugMessageData();
+
+	/**
+	 * A map of DebugMessageId's to their associated debug data.
+	 * The uint32 in the TTuple is the number of remaining message queues to be processed
+	 * for that message. It is decremented every time the message queue is processed, and when
+	 * it reaches zero is removed from this map.
+	 */
+	TMap<uint32, TTuple<uint32, FAsyncMessage::FMessageDebugData>> MessageDebugCallstacks;
+#endif
 };
diff --git a/Engine/Plugins/Experimental/AsyncMessageSystem/Source/AsyncMessageSystem/Public/AsyncMessageSystemBase.h b/Engine/Plugins/Experimental/AsyncMessageSystem/Source/AsyncMessageSystem/Public/AsyncMessageSystemBase.h
--- a/Engine/Plugins/Experimental/AsyncMessageSystem/Source/AsyncMessageSystem/Public/AsyncMessageSystemBase.h
+++ b/Engine/Plugins/Experimental/AsyncMessageSystem/Source/AsyncMessageSystem/Public/AsyncMessageSystemBase.h
@@ -414,18 +414,4 @@ FAsyncMessageHandle FAsyncMessageSystemBase::BindListener(
 		Options);
 
 	return bWasSuccessfullyBound ? NewListenerHandle : FAsyncMessageHandle::Invalid;
-}
-
-// A pre-processor define which will allow you to log the callstack of messages when they are queued.
-// This can be helpful when debugging a build to see what messages are being queued and from where,
-// since it can be hard to determine that information from the listener.
-//
-// By default, this will be enabled in the editor and non-shipping bulids.
-//
-// If you would like to enable this in a shipping build, then add a define to your game's .build.cs file:
-//		PublicDefinitions.Add("ENABLE_ASYNC_MESSAGES_DEBUG=1");
-//
-// It is not recommended that you log out on every message, because it will add significant performance cost.
-#ifndef ENABLE_ASYNC_MESSAGES_DEBUG
-	#define ENABLE_ASYNC_MESSAGES_DEBUG WITH_EDITOR || !UE_BUILD_SHIPPING
-#endif
\ No newline at end of file
+}
\ No newline at end of file
