[Async Message System]

Add "Message binding endpoints" to the system. These endpoints can be used to control the flow of messages and how listeners are set up.  When you queue a message for broadcast or bind a listener, you can specify a specific endpoiont. When a listener is bound to a specific endpoint it will only receive messages queued for that endpoint. This makes it simple and easy to, for example, set up an actor component with its own endpoint which can be used to notify it of gameplay events.

This also provides the advantage that the system does not need to iterate ALL bound listeners for each queued message, only the listeners who are bound to its associated end point.

#rnx

diff --git a/Engine/Plugins/Experimental/AsyncMessageSystem/Source/AsyncMessageSystem/Private/AsyncMessage.cpp b/Engine/Plugins/Experimental/AsyncMessageSystem/Source/AsyncMessageSystem/Private/AsyncMessage.cpp
--- a/Engine/Plugins/Experimental/AsyncMessageSystem/Source/AsyncMessageSystem/Private/AsyncMessage.cpp
+++ b/Engine/Plugins/Experimental/AsyncMessageSystem/Source/AsyncMessageSystem/Private/AsyncMessage.cpp
@@ -1,6 +1,7 @@
 // Copyright Epic Games, Inc. All Rights Reserved.
 
 #include "AsyncMessage.h"
+#include "AsyncMessageBindingEndpoint.h"
 
 #if ENABLE_ASYNC_MESSAGES_DEBUG
 #include "HAL/PlatformStackWalk.h"
@@ -13,7 +14,8 @@ FAsyncMessage::FAsyncMessage(
 	const uint64 CurrentFrame,
 	const uint32 InThreadQueuedFrom,
 	const uint32 InMessageSequenceCount,
-	FConstStructView InPayloadData)
+	FConstStructView InPayloadData,
+	TWeakPtr<FAsyncMessageBindingEndpoint> InBindingEndpoint)
 	: MessageId(InMessageId)
 	, MessageSourceId(InMessageSourceId)
 	, QueueTime(MessageTimestamp)
@@ -26,6 +28,7 @@ FAsyncMessage::FAsyncMessage(
 									// This is the same as : 
 									//		PayloadCopy.InitializeAs(InPayloadData.GetScriptStruct(), InPayloadData.GetMemory());
 									// TODO: Use some custom linear allocator here on the message system to make the copying cheaper
+	, BindingEndpoint(InBindingEndpoint)
 {
 
 }
@@ -75,6 +78,11 @@ FConstStructView FAsyncMessage::GetPayloadView() const
 	return FConstStructView(PayloadCopy);
 }
 
+TSharedPtr<FAsyncMessageBindingEndpoint> FAsyncMessage::GetBindingEndpoint() const
+{
+	return BindingEndpoint.Pin();
+}
+
 #if ENABLE_ASYNC_MESSAGES_DEBUG
 
 const FString& FAsyncMessage::FMessageDebugData::GetOrCreateNativeCallstackAsString() const
diff --git a/Engine/Plugins/Experimental/AsyncMessageSystem/Source/AsyncMessageSystem/Private/AsyncMessageBindingEndpoint.cpp b/Engine/Plugins/Experimental/AsyncMessageSystem/Source/AsyncMessageSystem/Private/AsyncMessageBindingEndpoint.cpp
new file mode 100644
--- /dev/null
+++ b/Engine/Plugins/Experimental/AsyncMessageSystem/Source/AsyncMessageSystem/Private/AsyncMessageBindingEndpoint.cpp
@@ -0,0 +1,50 @@
+// Copyright Epic Games, Inc. All Rights Reserved.
+
+#include "AsyncMessageBindingEndpoint.h"
+
+FAsyncMessageBindingEndpoint::FAsyncMessageBoundData* FAsyncMessageBindingEndpoint::GetBoundDataForMessage(const FAsyncMessageId& MessageId)
+{
+	return BoundMessageListenerMap.Find(MessageId);
+}
+
+FAsyncMessageBindingEndpoint::FAsyncMessageBoundData& FAsyncMessageBindingEndpoint::FindOrAddMessageData(const FAsyncMessageId& MessageId)
+{
+	return BoundMessageListenerMap.FindOrAdd(MessageId);
+}
+
+uint32 FAsyncMessageBindingEndpoint::GetNumberOfBoundListeners() const
+{
+	uint32 Res = 0;
+	
+	for (const TPair<FAsyncMessageId, FAsyncMessageBindingEndpoint::FAsyncMessageBoundData>& MessagePair : BoundMessageListenerMap)
+	{
+		const FAsyncMessageBindingEndpoint::FAsyncMessageBoundData& BoundData = MessagePair.Value;
+		for (const TPair<FAsyncMessageBindingOptions, TArray<FAsyncMessageBindingEndpoint::FAsyncMessageIndividualListener<FMessageCallbackFunc>>>& ListenerPair : BoundData.ListenerMap)
+		{
+			Res += ListenerPair.Value.Num();
+		}
+	}	
+
+	return Res;
+}
+
+bool FAsyncMessageBindingEndpoint::IsHandleBound(const FAsyncMessageHandle& Handle) const
+{
+	if (const FAsyncMessageBoundData* FoundBindingData = BoundMessageListenerMap.Find(Handle.GetBoundMessageId()))
+	{
+		for (const TPair<FAsyncMessageBindingOptions, TArray<FAsyncMessageIndividualListener<FMessageCallbackFunc>>>& ListenerPair : FoundBindingData->ListenerMap)
+		{
+			const bool bHasHandle = ListenerPair.Value.ContainsByPredicate([Handle](const FAsyncMessageIndividualListener<FMessageCallbackFunc>& Listener)
+			{
+				return Listener.Handle == Handle;	
+			});
+	
+			if (bHasHandle)
+			{
+				return true;
+			}
+		}	
+	}
+	
+	return false;
+}
diff --git a/Engine/Plugins/Experimental/AsyncMessageSystem/Source/AsyncMessageSystem/Private/AsyncMessageHandle.cpp b/Engine/Plugins/Experimental/AsyncMessageSystem/Source/AsyncMessageSystem/Private/AsyncMessageHandle.cpp
--- a/Engine/Plugins/Experimental/AsyncMessageSystem/Source/AsyncMessageSystem/Private/AsyncMessageHandle.cpp
+++ b/Engine/Plugins/Experimental/AsyncMessageSystem/Source/AsyncMessageSystem/Private/AsyncMessageHandle.cpp
@@ -1,6 +1,7 @@
 // Copyright Epic Games, Inc. All Rights Reserved.
 
 #include "AsyncMessageHandle.h"
+#include "AsyncMessageBindingEndpoint.h"
 
 const FAsyncMessageHandle FAsyncMessageHandle::Invalid = FAsyncMessageHandle();
 
@@ -26,12 +27,12 @@ FString FAsyncMessageHandle::ToString() const
 
 bool FAsyncMessageHandle::operator==(const FAsyncMessageHandle& Other) const
 {
-	return InternalHandle == Other.InternalHandle; 
+	return InternalHandle == Other.InternalHandle && BindingEndpoint == Other.BindingEndpoint; 
 }
 
 bool FAsyncMessageHandle::operator!=(const FAsyncMessageHandle& Other) const
 {
-	return InternalHandle != Other.InternalHandle; 
+	return !(*this == Other); 
 }
 
 bool FAsyncMessageHandle::operator>=(const FAsyncMessageHandle& Other) const
@@ -44,6 +45,11 @@ bool FAsyncMessageHandle::operator<(const FAsyncMessageHandle& Other) const
 	return InternalHandle < Other.InternalHandle;
 }
 
+TSharedPtr<FAsyncMessageBindingEndpoint> FAsyncMessageHandle::GetBindingEndpoint() const
+{
+	return BindingEndpoint.Pin();
+}
+
 uint32 GetTypeHash(const FAsyncMessageHandle& InMapping)
 {
 	return GetTypeHash(InMapping.InternalHandle);
@@ -51,10 +57,15 @@ uint32 GetTypeHash(const FAsyncMessageHandle& InMapping)
 
 //////////////////////////////////////////////////////////
 // Internal creation of message handles...
-FAsyncMessageHandle::FAsyncMessageHandle(const uint32 InHandleValue, const FAsyncMessageId InBoundMessage)
+FAsyncMessageHandle::FAsyncMessageHandle(
+	 const uint32 InHandleValue,
+	 const FAsyncMessageId InBoundMessage,
+	 TWeakPtr<FAsyncMessageBindingEndpoint> InBindingEndpoint)
 	: InternalHandle(InHandleValue)
 	, BoundMessage(InBoundMessage)
+	, BindingEndpoint(InBindingEndpoint) 
 {
 	checkf(InHandleValue != InvalidHandleIndex, TEXT("'%ud' is an invalid value for FAsyncMessageHandle!"), InHandleValue);
 	checkf(BoundMessage.IsValid(), TEXT("'%s' is an invalid FAsyncMessageId to create a handle for!"), *BoundMessage.ToString());
+	checkf(BindingEndpoint.IsValid(), TEXT("Message handle for message '%s' does not have a valid handler!"), *BoundMessage.ToString());
 }
\ No newline at end of file
diff --git a/Engine/Plugins/Experimental/AsyncMessageSystem/Source/AsyncMessageSystem/Private/AsyncMessageSystemBase.cpp b/Engine/Plugins/Experimental/AsyncMessageSystem/Source/AsyncMessageSystem/Private/AsyncMessageSystemBase.cpp
--- a/Engine/Plugins/Experimental/AsyncMessageSystem/Source/AsyncMessageSystem/Private/AsyncMessageSystemBase.cpp
+++ b/Engine/Plugins/Experimental/AsyncMessageSystem/Source/AsyncMessageSystem/Private/AsyncMessageSystemBase.cpp
@@ -3,6 +3,7 @@
 #include "AsyncMessageSystemBase.h"
 
 #include "AsyncMessageDeveloperSettings.h"
+#include "AsyncMessageBindingEndpoint.h"
 #include "AsyncMessageSystemLogs.h"
 #include "CoreGlobals.h"
 #include "Misc/App.h"
@@ -21,11 +22,16 @@ FAsyncMessageSystemBase::~FAsyncMessageSystemBase()
 FAsyncMessageHandle FAsyncMessageSystemBase::BindListener(
 	const FAsyncMessageId MessageId,
 	FMessageCallbackFunc&& Callback,
-	const FAsyncMessageBindingOptions& Options /* = {} */)
+	const FAsyncMessageBindingOptions& Options /* = {} */,
+	TWeakPtr<FAsyncMessageBindingEndpoint> BindingEndpoint /* = nullptr */)
 {
-	const FAsyncMessageHandle OutHandle = GenerateNextValidMessageHandle(MessageId);
+	const FAsyncMessageHandle OutHandle = GenerateNextValidMessageHandle(MessageId, BindingEndpoint);
 	
-	const bool bWasSuccessfullyBound = BindListener_Impl(OutHandle, MessageId, MoveTemp(Callback), Options);
+	const bool bWasSuccessfullyBound = BindListener_Impl(
+		OutHandle,
+		MessageId,
+		MoveTemp(Callback),
+		Options);
 	
 	return bWasSuccessfullyBound ? OutHandle : FAsyncMessageHandle::Invalid;
 }
@@ -35,7 +41,10 @@ void FAsyncMessageSystemBase::UnbindListener(const FAsyncMessageHandle& HandleTo
 	UnbindListener_Impl(HandleToUnbind);
 }
 
-bool FAsyncMessageSystemBase::QueueMessageForBroadcast(const FAsyncMessageId MessageId, FConstStructView PayloadData)
+bool FAsyncMessageSystemBase::QueueMessageForBroadcast(
+	const FAsyncMessageId MessageId,
+	FConstStructView PayloadData,
+	TWeakPtr<FAsyncMessageBindingEndpoint> BindingEndpoint /*= nullptr*/)
 {
 #if ENABLE_ASYNC_MESSAGES_DEBUG
 	const UAsyncMessageDeveloperSettings* Settings = GetDefault<UAsyncMessageDeveloperSettings>();
@@ -51,7 +60,7 @@ bool FAsyncMessageSystemBase::QueueMessageForBroadcast(const FAsyncMessageId Mes
 		// Print the BP callstack if desired
 		if (Settings->ShouldPrintScriptCallstackOnMessageQueue())
 		{
-			PrintScriptCallstack();	
+			PrintScriptCallstack();
 		}
 		
 		// Note: In case you forgot to enable the printing of the Blueprint VM callstack, you can use these commands
@@ -73,7 +82,18 @@ bool FAsyncMessageSystemBase::QueueMessageForBroadcast(const FAsyncMessageId Mes
 
 	TRACE_CPUPROFILER_EVENT_SCOPE(FAsyncMessageSystemBase::QueueMessageForBroadcast);
 	
-	return QueueMessageForBroadcast_Impl(MessageId, MoveTemp(PayloadData));
+	// If BindingEndpoint is not set, then use the default listen handler.
+	// We will require a default listen handler to be set on the message system before you can queue any messages
+	if (!DefaultBindingEndpoint.IsValid())
+	{
+		UE_LOG(LogAsyncMessageSystem, Error, TEXT("[%hs] Failed to queue message: There is no valid DefaultBindingEndpoint on this message system."), __func__);
+		return false;
+	}
+	
+	return QueueMessageForBroadcast_Impl(
+		MessageId,
+		MoveTemp(PayloadData),
+		BindingEndpoint.IsValid() ? BindingEndpoint : DefaultBindingEndpoint);
 }
 
 void FAsyncMessageSystemBase::ProcessMessageQueueForBinding(const FAsyncMessageBindingOptions& Options)
@@ -94,7 +114,10 @@ void FAsyncMessageSystemBase::Startup()
 	
 	UE_LOG(LogAsyncMessageSystem, Verbose, TEXT("[%hs] Message System Startup..."), __func__);
 	
-	checkf(bIsShuttingDown == false, TEXT("Attempting to restart a message system is not supported."));
+	checkf(bIsShuttingDown == false && DefaultBindingEndpoint.IsValid() == false, TEXT("Attempting to restart a message system is not supported."));
+
+	// Create a default handler
+	DefaultBindingEndpoint = MakeShared<FAsyncMessageBindingEndpoint>();
 	
 	Startup_Impl();
 }
@@ -109,6 +132,8 @@ void FAsyncMessageSystemBase::Shutdown()
 	// attempt to queue any more messages or start any new tasks
 	bIsShuttingDown = true;
 
+	DefaultBindingEndpoint.Reset();
+
 	Shutdown_Impl();
 }
 
@@ -126,7 +151,11 @@ void FAsyncMessageSystemBase::AddReferencedObjects(UObject* InReferencer, FRefer
 ///////////////////////////////////////////////////////////
 // Virtual Implementations
 
-bool FAsyncMessageSystemBase::BindListener_Impl(const FAsyncMessageHandle& HandleToBindTo, const FAsyncMessageId MessageId, FMessageCallbackFunc&& Callback, const FAsyncMessageBindingOptions& Options)
+bool FAsyncMessageSystemBase::BindListener_Impl(
+	const FAsyncMessageHandle& HandleToBindTo,
+	const FAsyncMessageId MessageId,
+	FMessageCallbackFunc&& Callback,
+	const FAsyncMessageBindingOptions& Options)
 {
 	TRACE_CPUPROFILER_EVENT_SCOPE(FAsyncMessageSystemBase::BindListener_Impl);
 	
@@ -164,8 +193,12 @@ bool FAsyncMessageSystemBase::BindListener_Impl(const FAsyncMessageHandle& Handl
 		FPendingBoundListener<FMessageCallbackFunc> 
 		{
 			.MessageId = MessageId, 
-			.Data = FAsyncMessageIndividualListener<FMessageCallbackFunc> { .Handle = HandleToBindTo, .Callback = MoveTempIfPossible(Callback) }, 
-			.Options = Options 
+			.Data = FAsyncMessageBindingEndpoint::FAsyncMessageIndividualListener<FMessageCallbackFunc>
+				{
+					.Handle = HandleToBindTo,
+					.Callback = MoveTempIfPossible(Callback)
+				}, 
+			.Options = Options
 		});
 
 	// Successfully queued for binding
@@ -180,7 +213,10 @@ void FAsyncMessageSystemBase::UnbindListener_Impl(const FAsyncMessageHandle& Han
 	}
 }
 
-bool FAsyncMessageSystemBase::QueueMessageForBroadcast_Impl(const FAsyncMessageId MessageId, FConstStructView PayloadData)
+bool FAsyncMessageSystemBase::QueueMessageForBroadcast_Impl(
+	const FAsyncMessageId MessageId,
+	FConstStructView PayloadData,
+	TWeakPtr<FAsyncMessageBindingEndpoint> BindingEndpoint)
 {
 	// Do not allow the queue of new listeners during the shutdown of this message system
 	if (bIsShuttingDown)
@@ -228,7 +264,8 @@ bool FAsyncMessageSystemBase::QueueMessageForBroadcast_Impl(const FAsyncMessageI
 			/* .QueueFrame= */ QueuedFrame,
 			/* .ThreadQueuedFrom= */ ThreadQueuedFrom,
 			/* .SequenceId= */ MessageSequence,		
-			/* Payload= */ MoveTemp(PayloadData)		// <--- Note: The payload data is going to be copied for each message in the queue
+			/* Payload= */ MoveTemp(PayloadData),		// <--- Note: The payload data is going to be copied for each message in the queue
+			/* BindingEndpoint= */ BindingEndpoint
 		);
 			
 	// Queue this message to the parent message
@@ -291,39 +328,47 @@ void FAsyncMessageSystemBase::ProcessMessageQueueForBinding_Impl(const FAsyncMes
 		//
 		//const double CurrentTime = FApp::GetCurrentTime();
 		//const uint64 CurrentFrame = GFrameCounter;
-		
-		// For each message in the hierarchy of this message...
-		FAsyncMessageId::WalkMessageHierarchy(Message.GetMessageId(), [&](const FAsyncMessageId CurrentMessageId)
+
+		TSharedPtr<FAsyncMessageBindingEndpoint> MessageHandler = Message.GetBindingEndpoint();
+
+		// Only process messages with valid handlers.
+		// A message handler might be invalid the handler goes out of scope
+		// (gets GC'd, or is otherwise destroyed after a message is queued)
+		if (MessageHandler.IsValid())
 		{
-			// Set this message ID to the current message id
-			Message.SetMessageId(CurrentMessageId);
-			
-			// If we know of any bound listeners to this message, notify them.
-			if (FAsyncMessageBoundData* BoundData = BoundMessageListenerMap.Find(Message.GetMessageId()))
+			// For each message in the hierarchy of this message...
+			FAsyncMessageId::WalkMessageHierarchy(Message.GetMessageId(), [&](const FAsyncMessageId CurrentMessageId)
 			{
-				// Look for any listeners associated with this binding option
-				if (TArray<FAsyncMessageIndividualListener<FMessageCallbackFunc>>* Listeners = BoundData->ListenerMap.Find(Options))
+				// Set this message ID to the current message id
+				Message.SetMessageId(CurrentMessageId);
+						
+				// If we know of any bound listeners to this message, notify them.
+				if (const FAsyncMessageBindingEndpoint::FAsyncMessageBoundData* BoundData = MessageHandler->GetBoundDataForMessage(Message.GetMessageId()))
 				{
-					for (const FAsyncMessageIndividualListener<FMessageCallbackFunc>& Listener : *Listeners)
+					// Look for any listeners associated with this binding option
+					if (const TArray<FAsyncMessageBindingEndpoint::FAsyncMessageIndividualListener<FMessageCallbackFunc>>* Listeners = BoundData->ListenerMap.Find(Options))
 					{
-						// As long as the callback for this listener is valid, broadcast it
-						if (Listener.Callback)
+						for (const FAsyncMessageBindingEndpoint::FAsyncMessageIndividualListener<FMessageCallbackFunc>& Listener : *Listeners)
 						{
-							// Broadcast to the listener!					
-							Listener.Callback(Message);
+							// As long as the callback for this listener is valid, broadcast it
+							if (Listener.Callback)
+							{
+								// Broadcast to the listener!					
+								Listener.Callback(Message);
 						
-							MessagesCalled++;	
-						}
-						// Otherwise, this callback is no long valid (the lambda may have gone out of scope)
-						// and we should unbind it so that we don't bother trying to call it again
-						else
-						{
-							UnbindListener(Listener.Handle);
+								MessagesCalled++;	
+							}
+							// Otherwise, this callback is no long valid (the lambda may have gone out of scope)
+							// and we should unbind it so that we don't bother trying to call it again
+							else
+							{
+								UnbindListener(Listener.Handle);
+							}
 						}
 					}
 				}
-			}
-		});
+			});	
+		}
 		
 		OptionalNextMessage = MessageStore.GetNextMessageForBindingOption(Options);
 	}
@@ -347,13 +392,13 @@ void FAsyncMessageSystemBase::ProcessUnbindHandleRequests()
 	while (HandleToRemoveOpt.IsSet())
 	{
 		FAsyncMessageHandle& HandleToRemove = HandleToRemoveOpt.GetValue();
-
-		if (FAsyncMessageBoundData* FoundBindingData = BoundMessageListenerMap.Find(HandleToRemove.GetBoundMessageId()))
+		
+		if (FAsyncMessageBindingEndpoint::FAsyncMessageBoundData* FoundBindingData = HandleToRemove.GetBindingEndpoint()->GetBoundDataForMessage(HandleToRemove.GetBoundMessageId()))
 		{
 			// Check the listener map for this message...
 			for (auto& BindingPair : FoundBindingData->ListenerMap)
 			{
-				TArray<FAsyncMessageIndividualListener<FMessageCallbackFunc>>& Listeners = BindingPair.Value;
+				TArray<FAsyncMessageBindingEndpoint::FAsyncMessageIndividualListener<FMessageCallbackFunc>>& Listeners = BindingPair.Value;
 
 				// Remove any listeners whose handles have been marked for unbinding
 				for (int32 i = Listeners.Num() - 1; i >= 0; --i)
@@ -389,7 +434,10 @@ void FAsyncMessageSystemBase::ProcessListenersPendingBinding()
 	TOptional<FPendingBoundListener<FMessageCallbackFunc>> ListenerToBind = PendingBoundListenerQueue.Dequeue();
 	while (ListenerToBind.IsSet())
 	{
-		FAsyncMessageBoundData& MessageData = BoundMessageListenerMap.FindOrAdd(ListenerToBind->MessageId);
+		TSharedPtr<FAsyncMessageBindingEndpoint> MessageHandler = ListenerToBind->Data.Handle.GetBindingEndpoint();
+		
+		// Bind a new listener to the message handler
+		FAsyncMessageBindingEndpoint::FAsyncMessageBoundData& MessageData = MessageHandler->FindOrAddMessageData(ListenerToBind->MessageId);
 
 		// Add a new listener which we can broadcast the callback to later
 		MessageData.ListenerMap.FindOrAdd(ListenerToBind->Options).Add(MoveTemp(ListenerToBind->Data));
@@ -398,7 +446,7 @@ void FAsyncMessageSystemBase::ProcessListenersPendingBinding()
 	}
 }
 
-FAsyncMessageHandle FAsyncMessageSystemBase::GenerateNextValidMessageHandle(const FAsyncMessageId ForMessageId)
+FAsyncMessageHandle FAsyncMessageSystemBase::GenerateNextValidMessageHandle(const FAsyncMessageId ForMessageId, TWeakPtr<FAsyncMessageBindingEndpoint> Endpoint)
 {
 	// Increment to the next handle, checking if its value is zero.
 	// If the value is zero, that means that we have hit uint32 max and overflowed back around
@@ -413,13 +461,13 @@ FAsyncMessageHandle FAsyncMessageSystemBase::GenerateNextValidMessageHandle(cons
 		UE_LOG(LogAsyncMessageSystem, Warning, TEXT("The Async Message Handle index has been wrapped!"));
 	}
 	
-	return FAsyncMessageHandle(HandleValue, ForMessageId);
+	return FAsyncMessageHandle(HandleValue, ForMessageId, Endpoint.IsValid() ? Endpoint: DefaultBindingEndpoint);
 }
 
 // Dev test function only to test creation of valid handles, just to make unit testing easier.
 #if WITH_DEV_AUTOMATION_TESTS
-FAsyncMessageHandle FAsyncMessageSystemBase::GenerateHandleAtIndex(const uint32 Index, const FAsyncMessageId& ForId)
+FAsyncMessageHandle FAsyncMessageSystemBase::GenerateHandleAtIndex(const uint32 Index, const FAsyncMessageId& ForId, TWeakPtr<FAsyncMessageBindingEndpoint> Endpoint)
 {
-	return FAsyncMessageHandle(Index, ForId);
+	return FAsyncMessageHandle(Index, ForId, Endpoint);
 }
 #endif
\ No newline at end of file
diff --git a/Engine/Plugins/Experimental/AsyncMessageSystem/Source/AsyncMessageSystem/Private/Tests/AsyncMessageSystemTests.cpp b/Engine/Plugins/Experimental/AsyncMessageSystem/Source/AsyncMessageSystem/Private/Tests/AsyncMessageSystemTests.cpp
--- a/Engine/Plugins/Experimental/AsyncMessageSystem/Source/AsyncMessageSystem/Private/Tests/AsyncMessageSystemTests.cpp
+++ b/Engine/Plugins/Experimental/AsyncMessageSystem/Source/AsyncMessageSystem/Private/Tests/AsyncMessageSystemTests.cpp
@@ -4,6 +4,7 @@
 
 #include "AsyncMessageHandle.h"
 #include "AsyncMessageSystemBase.h"
+#include "AsyncMessageBindingEndpoint.h"
 #include "Misc/AutomationTest.h"
 #include "NativeGameplayTags.h"
 #include "Tasks/Task.h"
@@ -41,9 +42,15 @@ namespace UE::AsyncMessageSystem
 		}
 		
 		bool bHasStarted = false;
-		bool bHasShutdown = false;	
-
+		bool bHasShutdown = false;
 	public:
+		[[nodiscard]] FAsyncMessageHandle GenerateHandleAtIndex(
+		   const uint32 Index,
+		   const FAsyncMessageId& ForId)
+		{
+			return FAsyncMessageSystemBase::GenerateHandleAtIndex(Index, ForId, DefaultBindingEndpoint);
+		}
+		
 		const bool HasStarted() const { return bHasStarted; }
 		const bool HasShutdown() const { return bHasShutdown && bIsShuttingDown; }
 		
@@ -53,43 +60,15 @@ namespace UE::AsyncMessageSystem
 		bool IsHandleBound(const FAsyncMessageHandle& Handle) const
 		{
 			FScopeLock ListenerLock(&MessageListenerMapCS);
-
-			if (const FAsyncMessageBoundData* FoundBindingData = BoundMessageListenerMap.Find(Handle.GetBoundMessageId()))
-			{
-				for (const TPair<FAsyncMessageBindingOptions, TArray<FAsyncMessageIndividualListener<FMessageCallbackFunc>>>& ListenerPair : FoundBindingData->ListenerMap)
-				{
-					const bool bHasHandle = ListenerPair.Value.ContainsByPredicate([Handle](const FAsyncMessageIndividualListener<FMessageCallbackFunc>& Listener)
-					{
-						return Listener.Handle == Handle;	
-					});
-		
-					if (bHasHandle)
-					{
-						return true;
-					}
-				}	
-			}
 			
-			return false;	
+			return DefaultBindingEndpoint->IsHandleBound(Handle);	
 		}
 
 		uint32 GetNumberOfListeners() const
 		{
-			uint32 Res = 0;
-			// Check any message listeners...
-			{
-				FScopeLock ListenerLock(&MessageListenerMapCS);
-				for (const TPair<FAsyncMessageId, FAsyncMessageBoundData>& MessagePair : BoundMessageListenerMap)
-				{
-					const FAsyncMessageBoundData& BoundData = MessagePair.Value;
-					for (const TPair<FAsyncMessageBindingOptions, TArray<FAsyncMessageIndividualListener<FMessageCallbackFunc>>>& ListenerPair : BoundData.ListenerMap)
-					{
-						Res += ListenerPair.Value.Num();
-					}
-				}	
-			}
+			FScopeLock ListenerLock(&MessageListenerMapCS);
 
-			return Res;
+			return DefaultBindingEndpoint->GetNumberOfBoundListeners();
 		}
 		
 		void RunOnce(const TArray<FAsyncMessageBindingOptions>& OptsToProcess = { FAsyncMessageBindingOptions{} })
@@ -107,7 +86,7 @@ namespace UE::AsyncMessageSystem
 
 		FAsyncMessageHandle GenerateTestHandle()
 		{
-			return GenerateNextValidMessageHandle({InternalTestTag_Invalid});
+			return GenerateNextValidMessageHandle({InternalTestTag_Invalid}, DefaultBindingEndpoint);
 		}
 	};
 	
@@ -115,6 +94,7 @@ namespace UE::AsyncMessageSystem
 	IMPLEMENT_SIMPLE_AUTOMATION_TEST(TMessageHandleDefaults, "AsyncMessagePassing.MessageHandles", QuickTestFlags)
 	bool TMessageHandleDefaults::RunTest(const FString& InParameters)
 	{
+		TSharedPtr<FTestMessageSystem, ESPMode::ThreadSafe> TestSystem = FAsyncMessageSystemBase::CreateMessageSystem<FTestMessageSystem>();
 		// Test defaults
 		{
 			FAsyncMessageHandle DefaultHandle = {};
@@ -127,14 +107,14 @@ namespace UE::AsyncMessageSystem
 
 		// Test Invalidate
 		{
-			FAsyncMessageHandle TestHandle = FAsyncMessageSystemBase::GenerateHandleAtIndex(123, {InternalTestTag_Invalid});
+			FAsyncMessageHandle TestHandle = TestSystem->GenerateHandleAtIndex(123, {InternalTestTag_Invalid});
 			TestTrue(TEXT("Test handle starts as a valid handle"), TestHandle.IsValid());
 		}
 
 		// Test GetId()
 		{
 			constexpr uint32 TestHandleIdx = 456;
-			FAsyncMessageHandle Handle = FAsyncMessageSystemBase::GenerateHandleAtIndex(TestHandleIdx, {InternalTestTag_Invalid});
+			FAsyncMessageHandle Handle = TestSystem->GenerateHandleAtIndex(TestHandleIdx, {InternalTestTag_Invalid});
 			TestTrue(TEXT("Handle Is Valid"), Handle.IsValid());
 			
 			TestTrue(TEXT("Handle Is the same as the given index"), TestHandleIdx == Handle.GetId());
@@ -143,7 +123,7 @@ namespace UE::AsyncMessageSystem
 		// Test ToString()
 		{
 			constexpr uint32 TestHandleIdx = 456;
-			FAsyncMessageHandle Handle = FAsyncMessageSystemBase::GenerateHandleAtIndex(TestHandleIdx, {InternalTestTag_Invalid});
+			FAsyncMessageHandle Handle = TestSystem->GenerateHandleAtIndex(TestHandleIdx, {InternalTestTag_Invalid});
 			
 			const FString DefaultHandleToString = Handle.ToString();
 			const FString ExpectedToString = TEXT("456");
@@ -152,18 +132,23 @@ namespace UE::AsyncMessageSystem
 
 		// Test handle == operator
 		{
-			FAsyncMessageHandle A = FAsyncMessageSystemBase::GenerateHandleAtIndex(789, {InternalTestTag_Invalid});
-			FAsyncMessageHandle B = FAsyncMessageSystemBase::GenerateHandleAtIndex(789, {InternalTestTag_Invalid});
+			FAsyncMessageHandle A = TestSystem->GenerateHandleAtIndex(789, {InternalTestTag_Invalid});
+			FAsyncMessageHandle B = TestSystem->GenerateHandleAtIndex(789, {InternalTestTag_Invalid});
 			TestTrue(TEXT("Async Message Handle == operator works"), A == B);
 		}
 
 		// Test handle != operator
 		{
-			FAsyncMessageHandle A = FAsyncMessageSystemBase::GenerateHandleAtIndex(123, {InternalTestTag_Invalid});
-			FAsyncMessageHandle B = FAsyncMessageSystemBase::GenerateHandleAtIndex(789, {InternalTestTag_Invalid});
+			FAsyncMessageHandle A = TestSystem->GenerateHandleAtIndex(123, {InternalTestTag_Invalid});
+			FAsyncMessageHandle B = TestSystem->GenerateHandleAtIndex(789, {InternalTestTag_Invalid});
 			TestTrue(TEXT("Async Message Handle != operator works"), A != B);
 		}
 
+		TestSystem->Shutdown();
+		TestTrue(TEXT("MessageSystem shutdown has been called"), TestSystem->HasShutdown());
+
+		TestSystem.Reset();
+		
 		return true;
 	}
 
@@ -464,6 +449,62 @@ namespace UE::AsyncMessageSystem
 		TestSystem.Reset();
 		return true;
 	}
+
+	// Test to queue messages from multiple threads on a single message system
+	IMPLEMENT_SIMPLE_AUTOMATION_TEST(TMessageSystemBindingEndpointsTest, "AsyncMessagePassing.Binding.CustomEndpoints", QuickTestFlags)
+	bool TMessageSystemBindingEndpointsTest::RunTest(const FString& InParameters)
+	{
+		TSharedPtr<FTestMessageSystem, ESPMode::ThreadSafe> TestSystem = FAsyncMessageSystemBase::CreateMessageSystem<FTestMessageSystem>();
+		TestTrue(TEXT("MessageSystem startup has been called"), TestSystem->HasStarted());
+
+		static const FAsyncMessageId TestMessageId { InternalTestTag_A };
+
+		FAsyncMessageBindingOptions BindOptions = {};
+		TArray<FAsyncMessageBindingOptions> BindingsToTick = { BindOptions };
+		
+		// Bind to the normal (default) endpoint
+		uint32 DefaultEndpointValue = 0;
+		{
+			TestSystem->BindListener(TestMessageId, [&DefaultEndpointValue](const FAsyncMessage& Message)
+			{
+				DefaultEndpointValue++;
+			},
+			BindOptions);
+		}
+
+		// Create a custom binding endpoint
+		TSharedPtr<FAsyncMessageBindingEndpoint> CustomEndpoint = MakeShared<FAsyncMessageBindingEndpoint>();
+		uint32 CustomEndpointValue = 0;
+		{
+			TestSystem->BindListener(TestMessageId, [&CustomEndpointValue](const FAsyncMessage& Message)
+			{
+				CustomEndpointValue++;
+			},
+			BindOptions,
+			CustomEndpoint);
+		}
+
+		// Queue a message for the default endpoint
+		TestSystem->QueueMessageForBroadcast(TestMessageId, /*payload*/nullptr);
+
+		// and tick the bindings
+		TestSystem->RunOnce(BindingsToTick);
+
+		TestEqual(TEXT("Default endpoint got called correctly"), DefaultEndpointValue, 1);
+		TestEqual(TEXT("CustomEndpoint did not get called, as expected"), CustomEndpointValue, 0);
+
+		// Queue a message for the custom endpoint...
+		TestSystem->QueueMessageForBroadcast(TestMessageId, /*payload*/nullptr, CustomEndpoint);
+		TestSystem->RunOnce(BindingsToTick);
+		
+		TestEqual(TEXT("Default endpoint did not get called, as expected"), DefaultEndpointValue, 1);
+		TestEqual(TEXT("CustomEndpoint was called correctly"), CustomEndpointValue, 1);
+		
+		TestSystem->Shutdown();
+		TestSystem.Reset();
+		
+		return true;
+	}
 	
 	// Test that the handle of an message is successfully unbound
 	IMPLEMENT_SIMPLE_AUTOMATION_TEST(TMessageSystemUnBindingSingle, "AsyncMessagePassing.UnBinding.Single", QuickTestFlags)
diff --git a/Engine/Plugins/Experimental/AsyncMessageSystem/Source/AsyncMessageSystem/Public/AsyncMessage.h b/Engine/Plugins/Experimental/AsyncMessageSystem/Source/AsyncMessageSystem/Public/AsyncMessage.h
--- a/Engine/Plugins/Experimental/AsyncMessageSystem/Source/AsyncMessageSystem/Public/AsyncMessage.h
+++ b/Engine/Plugins/Experimental/AsyncMessageSystem/Source/AsyncMessageSystem/Public/AsyncMessage.h
@@ -3,8 +3,9 @@
 #pragma once
 
 #include "AsyncMessageId.h"
-#include "HAL/Platform.h"	// For uint64 type declaration
+#include "HAL/Platform.h"				// For uint64 type declaration
 #include "StructUtils/StructView.h"
+#include "Templates/SharedPointer.h"	// For TWeakPtr
 
 #include "AsyncMessage.generated.h"
 
@@ -22,6 +23,8 @@
 	#define ENABLE_ASYNC_MESSAGES_DEBUG WITH_EDITOR || !UE_BUILD_SHIPPING
 #endif
 
+class FAsyncMessageBindingEndpoint;
+
 /**
  * A single async message which has been queued for broadcasting, and listeners are
  * interested in receiving the data from.
@@ -44,7 +47,8 @@ public:
 		const uint64 CurrentFrame,
 		const uint32 ThreadQueuedFrom,
 		const uint32 MessageSequenceCounter,
-		FConstStructView InPayloadData);
+		FConstStructView InPayloadData,
+		TWeakPtr<FAsyncMessageBindingEndpoint> InBindingEndpoint);
 
 	/**
 	 * Returns this message's unique ID which listeners use to bind to it. 
@@ -104,6 +108,10 @@ public:
 		return PayloadCopy.GetPtr<TDataType>();
 	}
 
+	/**
+	 * Returns the endpoint which this message should be broadcast on.
+	 */
+	TSharedPtr<FAsyncMessageBindingEndpoint> GetBindingEndpoint() const;
 	
 	FORCEINLINE void AddReferencedObjects(FReferenceCollector& Collector)
 	{
@@ -172,6 +180,14 @@ private:
 	UPROPERTY(BlueprintReadOnly, VisibleAnywhere, Category="Async Message", meta = (DisplayName="Payload", AllowPrivateAccess = "true"))
 	FInstancedStruct PayloadCopy;
 
+	/**
+	 * Pointer to which endpoint this message should be broadcast on.
+	 * 
+	 * Store this as a weak pointer because we don't want any messages to be keeping handlers
+	 * from GC'ing if they are destroyed over time
+	 */
+	TWeakPtr<FAsyncMessageBindingEndpoint> BindingEndpoint = nullptr;
+	
 	// TODO: We could perhaps define a "lifetime" concept on these messages, similar to how AnimNext's
 	// TraitEvents (TraitEvent.h) work, as well as consumption perhaps?
 
diff --git a/Engine/Plugins/Experimental/AsyncMessageSystem/Source/AsyncMessageSystem/Public/AsyncMessageBindingEndpoint.h b/Engine/Plugins/Experimental/AsyncMessageSystem/Source/AsyncMessageSystem/Public/AsyncMessageBindingEndpoint.h
new file mode 100644
--- /dev/null
+++ b/Engine/Plugins/Experimental/AsyncMessageSystem/Source/AsyncMessageSystem/Public/AsyncMessageBindingEndpoint.h
@@ -0,0 +1,98 @@
+// Copyright Epic Games, Inc. All Rights Reserved.
+
+#pragma once
+
+#include "AsyncMessage.h"
+#include "AsyncMessageBindingOptions.h"
+#include "AsyncMessageHandle.h"
+#include "AsyncMessageId.h"
+#include "Containers/Map.h"
+#include "Templates/SharedPointer.h"
+
+/**
+* The message binding endpoint is what stores the actual map of listeners and what Message Id's
+* they are bound to.
+*
+* When binding to a message or queueing a message for broadcast, you can specify the endpoint.
+* Listeners will only receive the messages queued for their own endpoints. This makes it easy to
+* filter messages to a specific endpoint.
+*
+* These endpoints can be utilized within the game frame in a rather customizable way.
+* One example would be to add create an Actor Component which has an endpoint. This way,
+* you could send messages directly to a single specific actor.
+*/
+class FAsyncMessageBindingEndpoint final :
+	public TSharedFromThis<FAsyncMessageBindingEndpoint>
+{
+public:
+
+	ASYNCMESSAGESYSTEM_API FAsyncMessageBindingEndpoint() = default;
+	ASYNCMESSAGESYSTEM_API virtual ~FAsyncMessageBindingEndpoint() = default;
+	
+	using FMessageCallbackFunc = TFunction<void(const FAsyncMessage&)>;
+	
+	/**
+	 * Data stored about a single individual listener of a message
+	 */
+	template <typename CallbackFuncSignature>
+	struct FAsyncMessageIndividualListener
+	{
+		/**
+		* The handle to this individual listener assigned when it is bound to this message system.
+		*/
+		FAsyncMessageHandle Handle;
+		
+		/**
+		 * The delegate that listeners can bind to if they want to be called when this message
+		 * is broadcast. 
+		 */
+		CallbackFuncSignature Callback;
+	};
+	
+	/**
+	 * Data store for everything related to a single message
+	 */
+	struct FAsyncMessageBoundData
+	{
+		/**
+		 * A map of listeners for this message, ordered by their binding options.
+		 * 
+		 * This lets us keep track of how many binders are listening for this message on a specific
+		 * tick group or thread id
+		 */
+		TMap<FAsyncMessageBindingOptions, TArray<FAsyncMessageIndividualListener<FMessageCallbackFunc>>> ListenerMap;
+	};
+
+	/**
+	 * Returns data about listeners bound to the given async message id.
+	 * 
+	 * @param MessageId The message id to find bindings for.
+	 * @return A pointer to the bound message data. Nullptr if there are no listeners bound to this message. 
+	 */
+	FAsyncMessageBoundData* GetBoundDataForMessage(const FAsyncMessageId& MessageId);
+
+	/**
+	 * Finds or adds a binding to the given message id.
+	 * @param MessageId 
+	 * @return 
+	 */
+	FAsyncMessageBoundData& FindOrAddMessageData(const FAsyncMessageId& MessageId);
+	
+	/**
+	 * @return The total number of bound listeners to this endpoint
+	 */
+	uint32 GetNumberOfBoundListeners() const;
+
+	/**
+	 * @param Handle The message handle to check if it is bound to this endpoint
+	 * @return True if the given handle is bound to this endpoint
+	 */
+	bool IsHandleBound(const FAsyncMessageHandle& Handle) const;
+
+private:
+
+	/**
+	 * Map of message ID's to their associated message data.
+	 */
+	TMap<FAsyncMessageId, FAsyncMessageBindingEndpoint::FAsyncMessageBoundData> BoundMessageListenerMap;
+};
diff --git a/Engine/Plugins/Experimental/AsyncMessageSystem/Source/AsyncMessageSystem/Public/AsyncMessageHandle.h b/Engine/Plugins/Experimental/AsyncMessageSystem/Source/AsyncMessageSystem/Public/AsyncMessageHandle.h
--- a/Engine/Plugins/Experimental/AsyncMessageSystem/Source/AsyncMessageSystem/Public/AsyncMessageHandle.h
+++ b/Engine/Plugins/Experimental/AsyncMessageSystem/Source/AsyncMessageSystem/Public/AsyncMessageHandle.h
@@ -5,9 +5,12 @@
 #include "AsyncMessageId.h"
 #include "Containers/UnrealString.h"	// For FString
 #include "HAL/Platform.h"				// For uint32 type declaration
+#include "Templates/SharedPointer.h"	// For TWeakPtr
 
 #include "AsyncMessageHandle.generated.h"
 
+class FAsyncMessageBindingEndpoint;
+
 /**
  * Handle used to identify a bound listener to an Async Message.
  * These handles are unique to each bound listener and created by the
@@ -53,6 +56,11 @@ public:
 	bool operator<(const FAsyncMessageHandle& Other) const;
 	ASYNCMESSAGESYSTEM_API friend uint32 GetTypeHash(const FAsyncMessageHandle& InMapping);
 
+	/**
+	 * @return A pointer to the endpoint which this handle is bound to.
+	 */
+	TSharedPtr<FAsyncMessageBindingEndpoint> GetBindingEndpoint() const;
+	
 	/**
 	 * Represents an invalid handle which cannot have any Message listeners bound to it
 	 * and means that no more Messages will be received. 
@@ -64,7 +72,10 @@ private:
 	/**
 	 * Private constructor for use by the message system
 	 */
-	FAsyncMessageHandle(const uint32 InHandleValue, const FAsyncMessageId BoundMessageId);
+	FAsyncMessageHandle(
+		const uint32 InHandleValue,
+		 const FAsyncMessageId BoundMessageId,
+		TWeakPtr<FAsyncMessageBindingEndpoint> InBindingEndpoint);
 
 	/**
 	 * Internal index for message handles which we consider to be invalid.
@@ -83,7 +94,8 @@ private:
 	UPROPERTY(VisibleAnywhere, Category="Async Message")
 	FAsyncMessageId BoundMessage = FAsyncMessageId::Invalid;
 
-	// TODO: Should these handles also have a pointer back to the message system that they
-	// were generated from? That might be useful in the context of having multiple different
-	// message systems perhaps...
+	/**
+	 * The endpoint which this handle is bound to.
+	 */
+	TWeakPtr<FAsyncMessageBindingEndpoint> BindingEndpoint = nullptr;
 };
\ No newline at end of file
diff --git a/Engine/Plugins/Experimental/AsyncMessageSystem/Source/AsyncMessageSystem/Public/AsyncMessageSystemBase.h b/Engine/Plugins/Experimental/AsyncMessageSystem/Source/AsyncMessageSystem/Public/AsyncMessageSystemBase.h
--- a/Engine/Plugins/Experimental/AsyncMessageSystem/Source/AsyncMessageSystem/Public/AsyncMessageSystemBase.h
+++ b/Engine/Plugins/Experimental/AsyncMessageSystem/Source/AsyncMessageSystem/Public/AsyncMessageSystemBase.h
@@ -4,16 +4,18 @@
 
 #include "AsyncMessage.h"
 #include "AsyncMessageHandle.h"
+#include "AsyncMessageBindingEndpoint.h"
 #include "AsyncMessageId.h"
 #include "AsyncMessageStore.h"
 #include <atomic>
-#include "Containers/Map.h"
 #include "Containers/MpscQueue.h"
 #include "StructUtils/StructView.h"
 #include "Templates/SharedPointer.h"
 #include "UObject/StrongObjectPtr.h"
 #include "UObject/WeakObjectPtrTemplates.h"
 
+class FAsyncMessageBindingEndpoint;
+
 /**
  * Base abstract implementation of an async message system which can be used to easily
  * pass signals and messages across different threads in Unreal.
@@ -59,7 +61,8 @@ public:
 	FAsyncMessageHandle BindListener(
 		const FAsyncMessageId MessageId,
 		FMessageCallbackFunc&& Callback,
-		const FAsyncMessageBindingOptions& Options = {});
+		const FAsyncMessageBindingOptions& Options = {},
+		TWeakPtr<FAsyncMessageBindingEndpoint> BindingEndpoint = nullptr);
 
 	/**
 	 * Binds the given callback to this message so that when an message of type MessageId is broadcast, it will be executed.
@@ -79,7 +82,8 @@ public:
 		const FAsyncMessageId MessageId,
 		TWeakObjectPtr<TOwner> WeakOwnerObject,
 		void(TOwner::* Callback)(const FAsyncMessage&),
-		const FAsyncMessageBindingOptions& Options = {});
+		const FAsyncMessageBindingOptions& Options = {},
+		TWeakPtr<FAsyncMessageBindingEndpoint> BindingEndpoint = nullptr);
 	
 	/**
 	 * Binds the given callback to this message so that when an message of type MessageId is broadcast, it will be executed.
@@ -99,7 +103,8 @@ public:
 		const FAsyncMessageId MessageId,
 		TWeakPtr<TOwner> WeakObject,
 		void(TOwner::* Callback)(const FAsyncMessage&),
-		const FAsyncMessageBindingOptions& Options = {});
+		const FAsyncMessageBindingOptions& Options = {},
+		TWeakPtr<FAsyncMessageBindingEndpoint> BindingEndpoint = nullptr);
 
 	/**
 	 * Unbinds the given listener from its message so that it will no longer receive callbacks
@@ -120,7 +125,10 @@ public:
 	 *
 	 * @return True if this message had an listeners bound to it and it was successfully queued
 	 */
-	bool QueueMessageForBroadcast(const FAsyncMessageId MessageId, FConstStructView PayloadData = nullptr);
+	bool QueueMessageForBroadcast(
+		const FAsyncMessageId MessageId,
+		FConstStructView PayloadData = nullptr,
+		TWeakPtr<FAsyncMessageBindingEndpoint> BindingEndpoint = nullptr);
 
 	/**
 	 * Processes all async message on within the queue for the given binding options.
@@ -165,7 +173,11 @@ protected:
 	* Handles binding the message callback to the given async message handle. This will return true if the binding was successful 
 	* and false if it failed for some reason (the message system could be shutting down, the MessageId could be invalid, etc.
 	*/
-	[[nodiscard]] virtual bool BindListener_Impl(const FAsyncMessageHandle& HandleToBindTo, const FAsyncMessageId MessageId, FMessageCallbackFunc&& Callback, const FAsyncMessageBindingOptions& Options);
+	[[nodiscard]] virtual bool BindListener_Impl(
+		const FAsyncMessageHandle& HandleToBindTo,
+		const FAsyncMessageId MessageId,
+		FMessageCallbackFunc&& Callback,
+		const FAsyncMessageBindingOptions& Options);
 	
 	/**
 	 * Marks the given handle as ready to be unbound. Enqueues this handle to a queue which will be 
@@ -178,7 +190,10 @@ protected:
 	* 
 	* Note: The payload data provided will be copied to make it safe for listeners on different threads. 
 	*/
-	virtual bool QueueMessageForBroadcast_Impl(const FAsyncMessageId MessageId, FConstStructView PayloadData);
+	virtual bool QueueMessageForBroadcast_Impl(
+		const FAsyncMessageId MessageId,
+		FConstStructView PayloadData,
+		TWeakPtr<FAsyncMessageBindingEndpoint> BindingEndpoint);
 	
 	/**
 	 * Called after an message was queued (from QueueMessageForBroadcast_Impl). This forces child classes of the message system
@@ -219,7 +234,9 @@ protected:
 	 * Generates a new async message handle by atomically incrementing the NextMessageHandleId.
 	 * If the new value of NextMessageHandleId is above the max allowed handle id, it will be wrapped.
 	 */
-	[[nodiscard]] FAsyncMessageHandle GenerateNextValidMessageHandle(const FAsyncMessageId ForMessageId);
+	[[nodiscard]] FAsyncMessageHandle GenerateNextValidMessageHandle(
+		const FAsyncMessageId ForMessageId,
+		TWeakPtr<FAsyncMessageBindingEndpoint> Endpoint);
 	
 	/**
 	 * Critical section for when the message listeners map is changed (listeners are bound/unbound)
@@ -246,44 +263,12 @@ protected:
 	 * A quick flag to check if this system is in the process of shutting down.
 	 */
 	std::atomic<bool> bIsShuttingDown = false;
-	
-	/**
-	 * Data stored about a single individual listener of a message
-	 */
-	template <typename CallbackFuncSignature>
-	struct FAsyncMessageIndividualListener
-	{
-		/**
-		* The handle to this individual listener assigned when it is bound to this message system.
-		*/
-		FAsyncMessageHandle Handle;
-		
-		/**
-		 * The delegate that listeners can bind to if they want to be called when this message
-		 * is broadcast. 
-		 */
-		CallbackFuncSignature Callback;
-	};
-	
-	/**
-	 * Data store for everything related to a single message
-	 */
-	struct FAsyncMessageBoundData
-	{
-		/**
-		 * A map of listeners for this message, ordered by their binding options.
-		 * 
-		 * This lets us keep track of how many binders are listening for this message on a specific
-		 * tick group or thread id
-		 */
-		TMap<FAsyncMessageBindingOptions, TArray<FAsyncMessageIndividualListener<FMessageCallbackFunc>>> ListenerMap;
-	};
-	
+
 	/**
-	 * Map of message ID's to their associated message data.
+	 * The default binding endpoint for this message system. When listeners are bound or messages are queued
+	 * which do not specify a specific endpoint, this one will be used.
 	 */
-	TMap<FAsyncMessageId, FAsyncMessageBoundData> BoundMessageListenerMap;
-
+	TSharedPtr<FAsyncMessageBindingEndpoint> DefaultBindingEndpoint = nullptr;
 
 	/**
 	 * Where we store the messages when they are queued for broadcasting and where we can
@@ -302,7 +287,7 @@ protected:
 		const FAsyncMessageId MessageId;
 
 		// The data for this message, populated in FAsyncMessageSystemBase::BindListener_Impl
-		FAsyncMessageIndividualListener<TCallbackFuncSignature> Data;
+		FAsyncMessageBindingEndpoint::FAsyncMessageIndividualListener<TCallbackFuncSignature> Data;
 
 		// Binding options that this listener requires
 		FAsyncMessageBindingOptions Options = {};
@@ -329,7 +314,10 @@ public:
     /**
      * A test-only helper function to quickly create an message handle from an index.
      */
-    [[nodiscard]] static FAsyncMessageHandle GenerateHandleAtIndex(const uint32 Index, const FAsyncMessageId& ForId);
+    [[nodiscard]] static FAsyncMessageHandle GenerateHandleAtIndex(
+    	const uint32 Index,
+    	const FAsyncMessageId& ForId,
+    	TWeakPtr<FAsyncMessageBindingEndpoint> Endpoint);
 #endif
 };
 
@@ -355,12 +343,13 @@ FAsyncMessageHandle FAsyncMessageSystemBase::BindListener(
 	const FAsyncMessageId MessageId,
 	TWeakObjectPtr<TOwner> WeakOwnerObject,
 	void(TOwner::* Callback)(const FAsyncMessage&),
-	const FAsyncMessageBindingOptions& Options)
+	const FAsyncMessageBindingOptions& Options,
+	TWeakPtr<FAsyncMessageBindingEndpoint> BindingEndpoint /*= nullptr*/)
 {
 	// This function is only for UObjects, because TWeakObjectPtr is a thing
 	static_assert(std::is_base_of<UObject, TOwner>::value, "TOwner not derived from UObject");
 
-	const FAsyncMessageHandle NewListenerHandle = GenerateNextValidMessageHandle(MessageId);
+	const FAsyncMessageHandle NewListenerHandle = GenerateNextValidMessageHandle(MessageId, BindingEndpoint);
 	
 	const bool bWasSuccessfullyBound = BindListener_Impl(
 		NewListenerHandle,
@@ -387,9 +376,10 @@ FAsyncMessageHandle FAsyncMessageSystemBase::BindListener(
 	const FAsyncMessageId MessageId,
 	TWeakPtr<TOwner> WeakObject,
 	void(TOwner::* Callback)(const FAsyncMessage&),
-	const FAsyncMessageBindingOptions& Options)
+	const FAsyncMessageBindingOptions& Options,
+	TWeakPtr<FAsyncMessageBindingEndpoint> BindingEndpoint /*= nullptr*/)
 {
-	const FAsyncMessageHandle NewListenerHandle = GenerateNextValidMessageHandle(MessageId);
+	const FAsyncMessageHandle NewListenerHandle = GenerateNextValidMessageHandle(MessageId, BindingEndpoint);
 	
 	const bool bWasSuccessfullyBound = BindListener_Impl(
 		NewListenerHandle,
diff --git a/Engine/Plugins/Experimental/AsyncMessageSystemTests/Source/AsyncMessageSystemTests/Private/AsyncMessageSystemPerformanceTests.cpp b/Engine/Plugins/Experimental/AsyncMessageSystemTests/Source/AsyncMessageSystemTests/Private/AsyncMessageSystemPerformanceTests.cpp
--- a/Engine/Plugins/Experimental/AsyncMessageSystemTests/Source/AsyncMessageSystemTests/Private/AsyncMessageSystemPerformanceTests.cpp
+++ b/Engine/Plugins/Experimental/AsyncMessageSystemTests/Source/AsyncMessageSystemTests/Private/AsyncMessageSystemPerformanceTests.cpp
@@ -59,7 +59,7 @@ void AASyncMessagePerfTest::EndPlay(const EEndPlayReason::Type EndPlayReason)
 	auto Sys = UAsyncMessageWorldSubsystem::GetSharedMessageSystem(GetWorld());
 	if (Sys.IsValid())
 	{
-		for (const FAsyncMessageHandle BoundHandle : BoundHandles)
+		for (const FAsyncMessageHandle& BoundHandle : BoundHandles)
 		{
 			Sys->UnbindListener(BoundHandle);
 		}
