Remove useless virtual destructor of FAsyncMessageBindingEndpoint. This got flagged by static analysis

#rnx

diff --git a/Engine/Plugins/Experimental/AsyncMessageSystem/Source/AsyncMessageSystem/Public/AsyncMessageBindingEndpoint.h b/Engine/Plugins/Experimental/AsyncMessageSystem/Source/AsyncMessageSystem/Public/AsyncMessageBindingEndpoint.h
--- a/Engine/Plugins/Experimental/AsyncMessageSystem/Source/AsyncMessageSystem/Public/AsyncMessageBindingEndpoint.h
+++ b/Engine/Plugins/Experimental/AsyncMessageSystem/Source/AsyncMessageSystem/Public/AsyncMessageBindingEndpoint.h
@@ -27,7 +27,6 @@ class FAsyncMessageBindingEndpoint final :
 public:
 
 	ASYNCMESSAGESYSTEM_API FAsyncMessageBindingEndpoint() = default;
-	ASYNCMESSAGESYSTEM_API virtual ~FAsyncMessageBindingEndpoint() = default;
 	
 	using FMessageCallbackFunc = TFunction<void(const FAsyncMessage&)>;
 	
