[Async message system]

Add BP support for binding endpoints via script interface.

Endpoints allow you to filter down which listeners you want to receive a message. When queueing or listening for a message, you can specify an endpoint. If you do, only listeners bound to that same endpoint will receive the message.

#rnx

diff --git a/Engine/Plugins/Experimental/AsyncMessageSystem/Source/AsyncMessageSystem/Private/AsyncAction_ListenForMessage.cpp b/Engine/Plugins/Experimental/AsyncMessageSystem/Source/AsyncMessageSystem/Private/AsyncAction_ListenForMessage.cpp
--- a/Engine/Plugins/Experimental/AsyncMessageSystem/Source/AsyncMessageSystem/Private/AsyncAction_ListenForMessage.cpp
+++ b/Engine/Plugins/Experimental/AsyncMessageSystem/Source/AsyncMessageSystem/Private/AsyncAction_ListenForMessage.cpp
@@ -3,6 +3,7 @@
 #include "AsyncAction_ListenForMessage.h"
 
 #include "AsyncGameplayMessageSystem.h"
+#include "AsyncMessageBindingComponent.h"
 #include "AsyncMessageWorldSubsystem.h"
 #include "AsyncMessageSystemLogs.h"
 #include "Engine/Engine.h"
@@ -10,6 +11,7 @@
 UAsyncAction_ListenForAsyncMessage* UAsyncAction_ListenForAsyncMessage::StartListeningForAsyncMessage(
 	UObject* WorldContextObject,
 	const FAsyncMessageId MessageId,
+	TScriptInterface<IAsyncMessageBindingEndpointInterface> DesiredEndpoint,
 	const TEnumAsByte<ETickingGroup> TickGroup)
 {
 	check(GEngine);
@@ -25,6 +27,7 @@ UAsyncAction_ListenForAsyncMessage* UAsyncAction_ListenForAsyncMessage::StartLis
 	// Create a new Async BP action to listen for a message
 	UAsyncAction_ListenForAsyncMessage* Action = NewObject<UAsyncAction_ListenForAsyncMessage>();
 	Action->WeakWorldPtr = World;
+	Action->DesiredEndpoint = DesiredEndpoint ? DesiredEndpoint->GetEndpoint() : nullptr;
 	Action->MessageToListenFor = MessageId;
 	Action->BindingOptions.SetTickGroup(TickGroup);
 	
@@ -75,7 +78,8 @@ void UAsyncAction_ListenForAsyncMessage::StartListeningForMessage()
 		MessageToListenFor,
 		TWeakObjectPtr<UAsyncAction_ListenForAsyncMessage>{ this },
 		&UAsyncAction_ListenForAsyncMessage::HandleMessageReceived,
-		BindingOptions);
+		BindingOptions,
+		DesiredEndpoint);
 
 	// If for some reason we failed to bind the handle, then destroy this async action and log an error
 	if (!BoundListenerHandle.IsValid())
diff --git a/Engine/Plugins/Experimental/AsyncMessageSystem/Source/AsyncMessageSystem/Private/AsyncAction_ListenForMessage.h b/Engine/Plugins/Experimental/AsyncMessageSystem/Source/AsyncMessageSystem/Private/AsyncAction_ListenForMessage.h
--- a/Engine/Plugins/Experimental/AsyncMessageSystem/Source/AsyncMessageSystem/Private/AsyncAction_ListenForMessage.h
+++ b/Engine/Plugins/Experimental/AsyncMessageSystem/Source/AsyncMessageSystem/Private/AsyncAction_ListenForMessage.h
@@ -5,6 +5,7 @@
 #include "AsyncMessageBindingOptions.h"
 #include "AsyncMessageHandle.h"
 #include "AsyncMessageId.h"
+#include "AsyncMessageBindingComponent.h"
 #include "Kismet/BlueprintAsyncActionBase.h"
 #include "StructUtils/InstancedStruct.h"
 
@@ -38,12 +39,14 @@ public:
 	 * Starts listening for an Async Message with the given ID during the given tick group. 
 	 * 
 	 * @param MessageId The ID of the Async Message which you would like to listen to
+	 * @param DesiredEndpoint The endpoint which this listener should bind to. If nothing is provided, the default world endpoint will be used.
 	 * @param TickGroup The tick group which you would like to receive the message in. Default is TG_PostUpdateWork.
 	 */
 	UFUNCTION(BlueprintCallable, Category = "Async Messages", meta=(WorldContext="WorldContextObject", BlueprintInternalUseOnly="true", AdvancedDisplay="TickGroup", DisplayName="Start Listening for Async Message"))
 	static UAsyncAction_ListenForAsyncMessage* StartListeningForAsyncMessage(
 		UObject* WorldContextObject,
 		const FAsyncMessageId MessageId,
+		TScriptInterface<IAsyncMessageBindingEndpointInterface> DesiredEndpoint,
 		const TEnumAsByte<ETickingGroup> TickGroup = TG_PostUpdateWork);
 
 	/**
@@ -95,6 +98,11 @@ protected:
 	* Weak pointer to the owning wold context to which this listener belongs 
 	*/
 	TWeakObjectPtr<UWorld> WeakWorldPtr = nullptr;
+
+	/**
+	 * The specific endpoint to start listening for this message on.
+	 */
+	TWeakPtr<FAsyncMessageBindingEndpoint> DesiredEndpoint = nullptr;
 	
 	/**
 	* The async message id that this action should listen 
diff --git a/Engine/Plugins/Experimental/AsyncMessageSystem/Source/AsyncMessageSystem/Private/AsyncMessageBindingComponent.cpp b/Engine/Plugins/Experimental/AsyncMessageSystem/Source/AsyncMessageSystem/Private/AsyncMessageBindingComponent.cpp
new file mode 100644
--- /dev/null
+++ b/Engine/Plugins/Experimental/AsyncMessageSystem/Source/AsyncMessageSystem/Private/AsyncMessageBindingComponent.cpp
@@ -0,0 +1,39 @@
+// Copyright Epic Games, Inc. All Rights Reserved.
+
+#include "AsyncMessageBindingComponent.h"
+#include "AsyncMessageBindingEndpoint.h"
+
+void UAsyncMessageBindingComponent::BeginPlay()
+{
+	Super::BeginPlay();
+
+	CreateEndpoint();
+}
+
+void UAsyncMessageBindingComponent::EndPlay(const EEndPlayReason::Type EndPlayReason)
+{
+	Super::EndPlay(EndPlayReason);
+
+	CleanupEndpoint();
+}
+
+TSharedPtr<FAsyncMessageBindingEndpoint> UAsyncMessageBindingComponent::GetEndpoint() const
+{
+	return Endpoint;
+}
+
+void UAsyncMessageBindingComponent::CreateEndpoint()
+{
+	// If there is already a valid endpoint, then there is no need to do anything.
+	if (Endpoint.IsValid())
+	{
+		return;
+	}
+	
+	Endpoint = MakeShared<FAsyncMessageBindingEndpoint>();
+}
+
+void UAsyncMessageBindingComponent::CleanupEndpoint()
+{
+	Endpoint.Reset();
+}
diff --git a/Engine/Plugins/Experimental/AsyncMessageSystem/Source/AsyncMessageSystem/Private/AsyncMessageBlueprintLibrary.cpp b/Engine/Plugins/Experimental/AsyncMessageSystem/Source/AsyncMessageSystem/Private/AsyncMessageBlueprintLibrary.cpp
--- a/Engine/Plugins/Experimental/AsyncMessageSystem/Source/AsyncMessageSystem/Private/AsyncMessageBlueprintLibrary.cpp
+++ b/Engine/Plugins/Experimental/AsyncMessageSystem/Source/AsyncMessageSystem/Private/AsyncMessageBlueprintLibrary.cpp
@@ -5,12 +5,13 @@
 #include "AsyncMessageWorldSubsystem.h"
 #include "AsyncMessageSystemBase.h"
 #include "AsyncMessageSystemLogs.h"
-#include "Engine/Engine.h"	// For GEnngine 
+#include "Engine/Engine.h"	// For GEngine 
 
 bool UAsyncMessageSystemBlueprintLibrary::QueueAsyncMessageForBroadcast(
 	UObject* WorldContextObject,
 	const FAsyncMessageId& MessageId,
-	const FInstancedStruct& Payload)
+	const FInstancedStruct& Payload,
+	TScriptInterface<IAsyncMessageBindingEndpointInterface> DesiredEndpoint)
 {
 	check(GEngine);
 	UWorld* World = GEngine->GetWorldFromContextObject(WorldContextObject, EGetWorldErrorMode::LogAndReturnNull);
@@ -31,7 +32,9 @@ bool UAsyncMessageSystemBlueprintLibrary::QueueAsyncMessageForBroadcast(
 		return false;
 	}
 
-	return Sys->QueueMessageForBroadcast(MessageId, Payload);
+	TWeakPtr<FAsyncMessageBindingEndpoint> WeakEndpoint = DesiredEndpoint ? DesiredEndpoint->GetEndpoint() : nullptr;
+
+	return Sys->QueueMessageForBroadcast(MessageId, Payload, WeakEndpoint);
 }
 
 FString UAsyncMessageSystemBlueprintLibrary::Conv_AsyncMessageIdToString(const FAsyncMessageId& MessageId)
diff --git a/Engine/Plugins/Experimental/AsyncMessageSystem/Source/AsyncMessageSystem/Private/AsyncMessageBlueprintLibrary.h b/Engine/Plugins/Experimental/AsyncMessageSystem/Source/AsyncMessageSystem/Private/AsyncMessageBlueprintLibrary.h
--- a/Engine/Plugins/Experimental/AsyncMessageSystem/Source/AsyncMessageSystem/Private/AsyncMessageBlueprintLibrary.h
+++ b/Engine/Plugins/Experimental/AsyncMessageSystem/Source/AsyncMessageSystem/Private/AsyncMessageBlueprintLibrary.h
@@ -5,6 +5,7 @@
 #include "AsyncMessageId.h"
 #include "InstancedStruct.h"
 #include "Kismet/BlueprintFunctionLibrary.h"
+#include "AsyncMessageBindingComponent.h"
 
 #include "AsyncMessageBlueprintLibrary.generated.h"
 
@@ -24,15 +25,17 @@ public:
 	 *
 	 * @param MessageId		The message that you would like to queue for broadcasting
 	 * @param Payload		The payload data of this message. This payload data will be COPIED to the message queue to
-	 *						to make it safe for listeners on other threads. 
+	 *						to make it safe for listeners on other threads.
+	 * @param DesiredEndpoint The endpoint which this listener should bind to. If nothing is provided, the default world endpoint will be used. 
 	 *
 	 * @return True if this message had an listeners bound to it and it was successfully queued
 	 */
-	UFUNCTION(BlueprintCallable, Category = "Async Messages", meta=(WorldContext="WorldContextObject", ReturnDisplayName = "Success"))
+	UFUNCTION(BlueprintCallable, Category = "Async Messages", meta=(WorldContext="WorldContextObject", AutoCreateRefTerm = "Payload" ,ReturnDisplayName = "Success"))
 	static bool QueueAsyncMessageForBroadcast(
 		UObject* WorldContextObject,
 		const FAsyncMessageId& MessageId,
-		const FInstancedStruct& Payload);
+		const FInstancedStruct& Payload,
+		TScriptInterface<IAsyncMessageBindingEndpointInterface> DesiredEndpoint);
 
 	/**
 	 * Get the string representation of the given Async Message Id.
diff --git a/Engine/Plugins/Experimental/AsyncMessageSystem/Source/AsyncMessageSystem/Public/AsyncMessageBindingComponent.h b/Engine/Plugins/Experimental/AsyncMessageSystem/Source/AsyncMessageSystem/Public/AsyncMessageBindingComponent.h
new file mode 100644
--- /dev/null
+++ b/Engine/Plugins/Experimental/AsyncMessageSystem/Source/AsyncMessageSystem/Public/AsyncMessageBindingComponent.h
@@ -0,0 +1,66 @@
+// Copyright Epic Games, Inc. All Rights Reserved.
+
+#pragma once
+
+#include "UObject/Interface.h"
+#include "Components/ActorComponent.h"
+#include "Templates/SharedPointer.h"
+
+#include "AsyncMessageBindingComponent.generated.h"
+
+class FAsyncMessageBindingEndpoint;
+
+UINTERFACE(MinimalAPI, BlueprintType, meta = (CannotImplementInterfaceInBlueprint))
+class UAsyncMessageBindingEndpointInterface : public UInterface
+{
+	GENERATED_BODY()
+};
+
+/**
+ * Interface to implement which will allow users to specify which binding endpoint
+ * should be used in any script functions for the async message system.
+ */
+class ASYNCMESSAGESYSTEM_API IAsyncMessageBindingEndpointInterface
+{
+	GENERATED_BODY()
+
+public:
+	
+	/**
+	 * @return Pointer to the specific binding endpoint which you would like to set for your message or listener.
+	 */
+	virtual TSharedPtr<FAsyncMessageBindingEndpoint> GetEndpoint() const
+	{
+		return nullptr;
+	}
+};
+
+/**
+ * A blueprint component which will allow you to specify a specific endpoint to use
+ * when queueing messages or binding listeners to the Async Message System.
+ *
+ * By default, this will create an endpoint on BeginPlay and destroy it on EndPlay.
+ */
+UCLASS(MinimalAPI, BlueprintType, HideCategories=(Object,Replication,Navigation), Meta=(BlueprintSpawnableComponent))
+class UAsyncMessageBindingComponent final :
+	public UActorComponent,
+	public IAsyncMessageBindingEndpointInterface
+{
+	GENERATED_BODY()
+
+protected:
+
+	//~ Begin UActorComponent interface
+	virtual void BeginPlay() override;
+	virtual void EndPlay(const EEndPlayReason::Type EndPlayReason) override;
+	//~ end UActorComponent interface
+
+	//~ Begin IAsyncMessageBindingEndpointInterface
+	virtual TSharedPtr<FAsyncMessageBindingEndpoint> GetEndpoint() const override;
+	//~ End IAsyncMessageBindingEndpointInterface
+
+	ASYNCMESSAGESYSTEM_API virtual void CreateEndpoint();
+	ASYNCMESSAGESYSTEM_API virtual void CleanupEndpoint();
+	
+	TSharedPtr<FAsyncMessageBindingEndpoint> Endpoint = nullptr;
+};
