[Async Message System]

Fix CIS Static Analysis warning. Remove final from UAsyncMessageBindingComponent so that you override Create/Cleanup endpoint functions

#rnx

diff --git a/Engine/Plugins/Experimental/AsyncMessageSystem/Source/AsyncMessageSystem/Public/AsyncMessageBindingComponent.h b/Engine/Plugins/Experimental/AsyncMessageSystem/Source/AsyncMessageSystem/Public/AsyncMessageBindingComponent.h
--- a/Engine/Plugins/Experimental/AsyncMessageSystem/Source/AsyncMessageSystem/Public/AsyncMessageBindingComponent.h
+++ b/Engine/Plugins/Experimental/AsyncMessageSystem/Source/AsyncMessageSystem/Public/AsyncMessageBindingComponent.h
@@ -42,7 +42,7 @@ public:
  * By default, this will create an endpoint on BeginPlay and destroy it on EndPlay.
  */
 UCLASS(MinimalAPI, BlueprintType, HideCategories=(Object,Replication,Navigation), Meta=(BlueprintSpawnableComponent))
-class UAsyncMessageBindingComponent final :
+class UAsyncMessageBindingComponent :
 	public UActorComponent,
 	public IAsyncMessageBindingEndpointInterface
 {
