[Async message system] Expose FAsyncMessageBindingEndpoint with the _API to make it possible to use it when implementing a custom subclass of the message system.

#rnx

diff --git a/Engine/Plugins/Experimental/AsyncMessageSystem/Source/AsyncMessageSystem/Public/AsyncMessageBindingEndpoint.h b/Engine/Plugins/Experimental/AsyncMessageSystem/Source/AsyncMessageSystem/Public/AsyncMessageBindingEndpoint.h
--- a/Engine/Plugins/Experimental/AsyncMessageSystem/Source/AsyncMessageSystem/Public/AsyncMessageBindingEndpoint.h
+++ b/Engine/Plugins/Experimental/AsyncMessageSystem/Source/AsyncMessageSystem/Public/AsyncMessageBindingEndpoint.h
@@ -21,12 +21,12 @@
 * One example would be to add create an Actor Component which has an endpoint. This way,
 * you could send messages directly to a single specific actor.
 */
-class FAsyncMessageBindingEndpoint final :
+class ASYNCMESSAGESYSTEM_API FAsyncMessageBindingEndpoint final :
 	public TSharedFromThis<FAsyncMessageBindingEndpoint>
 {
 public:
 
-	ASYNCMESSAGESYSTEM_API FAsyncMessageBindingEndpoint() = default;
+	FAsyncMessageBindingEndpoint() = default;
 	
 	using FMessageCallbackFunc = TFunction<void(const FAsyncMessage&)>;
 	
