[Async Message System] Mark UAsyncMessageWorldSubsystem as a non-Bp type so that it doesn't show up needlessly in the BP context menu. All async message stuff in BP's should be interfaces through the BP function library because it is easier to use.

#rnx

diff --git a/Engine/Plugins/Experimental/AsyncMessageSystem/Source/AsyncMessageSystem/Public/AsyncMessageWorldSubsystem.h b/Engine/Plugins/Experimental/AsyncMessageSystem/Source/AsyncMessageSystem/Public/AsyncMessageWorldSubsystem.h
--- a/Engine/Plugins/Experimental/AsyncMessageSystem/Source/AsyncMessageSystem/Public/AsyncMessageWorldSubsystem.h
+++ b/Engine/Plugins/Experimental/AsyncMessageSystem/Source/AsyncMessageSystem/Public/AsyncMessageWorldSubsystem.h
@@ -19,7 +19,7 @@ class FAsyncMessageSystemBase;
  *		TSharedPtr<FAsyncMessageSystemBase> Sys = UAsyncMessageWorldSubsystem::GetSharedMessageSystem(GetWorld());
  *		Sys->QueueMessageForBroadcast(MessageToQueue, MessagePayloadToQueue);
  */
-UCLASS()
+UCLASS(NotBlueprintType)
 class ASYNCMESSAGESYSTEM_API UAsyncMessageWorldSubsystem :
 	public UWorldSubsystem
 {
