[Async Message System] Fix a crash if an endpoint has gone out of scope

#rnx

diff --git a/Engine/Plugins/Experimental/AsyncMessageSystem/Source/AsyncMessageSystem/Private/AsyncMessageSystemBase.cpp b/Engine/Plugins/Experimental/AsyncMessageSystem/Source/AsyncMessageSystem/Private/AsyncMessageSystemBase.cpp
--- a/Engine/Plugins/Experimental/AsyncMessageSystem/Source/AsyncMessageSystem/Private/AsyncMessageSystemBase.cpp
+++ b/Engine/Plugins/Experimental/AsyncMessageSystem/Source/AsyncMessageSystem/Private/AsyncMessageSystemBase.cpp
@@ -393,20 +393,23 @@ void FAsyncMessageSystemBase::ProcessUnbindHandleRequests()
 	{
 		FAsyncMessageHandle& HandleToRemove = HandleToRemoveOpt.GetValue();
 		
-		if (FAsyncMessageBindingEndpoint::FAsyncMessageBoundData* FoundBindingData = HandleToRemove.GetBindingEndpoint()->GetBoundDataForMessage(HandleToRemove.GetBoundMessageId()))
+		if (TSharedPtr<FAsyncMessageBindingEndpoint> Endpoint = HandleToRemove.GetBindingEndpoint())
 		{
-			// Check the listener map for this message...
-			for (auto& BindingPair : FoundBindingData->ListenerMap)
+			if (FAsyncMessageBindingEndpoint::FAsyncMessageBoundData* FoundBindingData = Endpoint->GetBoundDataForMessage(HandleToRemove.GetBoundMessageId()))
 			{
-				TArray<FAsyncMessageBindingEndpoint::FAsyncMessageIndividualListener<FMessageCallbackFunc>>& Listeners = BindingPair.Value;
-
-				// Remove any listeners whose handles have been marked for unbinding
-				for (int32 i = Listeners.Num() - 1; i >= 0; --i)
+				// Check the listener map for this message...
+				for (auto& BindingPair : FoundBindingData->ListenerMap)
 				{
-					if (Listeners[i].Handle == HandleToRemove)
+					TArray<FAsyncMessageBindingEndpoint::FAsyncMessageIndividualListener<FMessageCallbackFunc>>& Listeners = BindingPair.Value;
+
+					// Remove any listeners whose handles have been marked for unbinding
+					for (int32 i = Listeners.Num() - 1; i >= 0; --i)
 					{
-						MessageStore.RemoveMessageFromBinding(Listeners[i].Handle.GetBoundMessageId(), BindingPair.Key);
-						Listeners.RemoveAt(i);
+						if (Listeners[i].Handle == HandleToRemove)
+						{
+							MessageStore.RemoveMessageFromBinding(Listeners[i].Handle.GetBoundMessageId(), BindingPair.Key);
+							Listeners.RemoveAt(i);
+						}
 					}
 				}
 			}
