[Async Message System] Fix a bug in the Async Message System where, if you called bind/unbind in the same tick group (i.e. immediately in a BP) then the function would be bound anyway.

#rnx
#tests All Async Message System unit tests as well as the Lyra test map for it.
#Synced-CL 48353899

diff --git a/Engine/Plugins/Experimental/AsyncMessageSystem/Source/AsyncMessageSystem/Private/AsyncAction_ListenForMessage.cpp b/Engine/Plugins/Experimental/AsyncMessageSystem/Source/AsyncMessageSystem/Private/AsyncAction_ListenForMessage.cpp
--- a/Engine/Plugins/Experimental/AsyncMessageSystem/Source/AsyncMessageSystem/Private/AsyncAction_ListenForMessage.cpp
+++ b/Engine/Plugins/Experimental/AsyncMessageSystem/Source/AsyncMessageSystem/Private/AsyncAction_ListenForMessage.cpp
@@ -138,6 +138,8 @@ void UAsyncAction_ListenForAsyncMessage::UnbindListener()
 
 void UAsyncAction_ListenForAsyncMessage::HandleMessageReceived(const FAsyncMessage& Message)
 {
+	ensure(BoundListenerHandle.IsValid());
+	
 	OnMessageReceived.Broadcast(Message);
 }
 
diff --git a/Engine/Plugins/Experimental/AsyncMessageSystem/Source/AsyncMessageSystem/Private/AsyncMessageSystemBase.cpp b/Engine/Plugins/Experimental/AsyncMessageSystem/Source/AsyncMessageSystem/Private/AsyncMessageSystemBase.cpp
--- a/Engine/Plugins/Experimental/AsyncMessageSystem/Source/AsyncMessageSystem/Private/AsyncMessageSystemBase.cpp
+++ b/Engine/Plugins/Experimental/AsyncMessageSystem/Source/AsyncMessageSystem/Private/AsyncMessageSystemBase.cpp
@@ -413,6 +413,18 @@ void FAsyncMessageSystemBase::ProcessUnbindHandleRequests()
 					}
 				}
 			}
+			else
+			{
+				// If we get here, then the current handle has been unbound before we have actually had a chance to fully set up its binding yet.
+				// For example, if you call Bind/Unbind immediately in a blueprint node. 
+				// Keep track of these handles which we couldn't fully unbind yet so that we don't bother to set up their bindings.
+				UE_LOG(LogAsyncMessageSystem, Verbose, TEXT("Handle: %s -- No binding data found. That means this listener has not been fully added yet. We will prevent it from being added instead of removing it."), *HandleToRemove.ToString());				
+				IncompleteBoundHandlesToRemove.Add(HandleToRemove);
+			}
+		}
+		else
+		{
+			UE_LOG(LogAsyncMessageSystem, Error, TEXT("Attempting to unbind from an invalid endpoint. Failed to remove listeners for handle: %s"), *HandleToRemove.ToString());
 		}
 
 		// Get the next handle to remove
@@ -438,15 +450,29 @@ void FAsyncMessageSystemBase::ProcessListenersPendingBinding()
 	while (ListenerToBind.IsSet())
 	{
 		TSharedPtr<FAsyncMessageBindingEndpoint> MessageHandler = ListenerToBind->Data.Handle.GetBindingEndpoint();
-		
-		// Bind a new listener to the message handler
-		FAsyncMessageBindingEndpoint::FAsyncMessageBoundData& MessageData = MessageHandler->FindOrAddMessageData(ListenerToBind->MessageId);
 
-		// Add a new listener which we can broadcast the callback to later
-		MessageData.ListenerMap.FindOrAdd(ListenerToBind->Options).Add(MoveTemp(ListenerToBind->Data));
+		// Only finish setting up the bindings for handles which have not been unbound already.
+		if (!IncompleteBoundHandlesToRemove.Contains(ListenerToBind->Data.Handle))
+		{
+			// Bind a new listener to the message handler
+			FAsyncMessageBindingEndpoint::FAsyncMessageBoundData& MessageData = MessageHandler->FindOrAddMessageData(ListenerToBind->MessageId);
+
+			// Add a new listener which we can broadcast the callback to later
+			MessageData.ListenerMap.FindOrAdd(ListenerToBind->Options).Add(MoveTemp(ListenerToBind->Data));
+		}				
+		else
+		{
+			UE_LOG(LogAsyncMessageSystem, Verbose, TEXT("Skipping pending binding to handle %s because it was already marked for unbinding."), *ListenerToBind->Data.Handle.ToString());
+		}
 
 		ListenerToBind = PendingBoundListenerQueue.Dequeue();
 	}
+
+	// We have processed all the pending bound listeners so we can now empty this
+	IncompleteBoundHandlesToRemove.Empty();
+
+	// After iterating and dequeue'ing off the listener map it should be empty
+	check(PendingBoundListenerQueue.IsEmpty());
 }
 
 FAsyncMessageHandle FAsyncMessageSystemBase::GenerateNextValidMessageHandle(const FAsyncMessageId ForMessageId, TWeakPtr<FAsyncMessageBindingEndpoint> Endpoint)
@@ -464,7 +490,7 @@ FAsyncMessageHandle FAsyncMessageSystemBase::GenerateNextValidMessageHandle(cons
 		UE_LOG(LogAsyncMessageSystem, Warning, TEXT("The Async Message Handle index has been wrapped!"));
 	}
 	
-	return FAsyncMessageHandle(HandleValue, ForMessageId, Endpoint.IsValid() ? Endpoint: DefaultBindingEndpoint);
+	return FAsyncMessageHandle(HandleValue, ForMessageId, Endpoint.IsValid() ? Endpoint : DefaultBindingEndpoint);
 }
 
 // Dev test function only to test creation of valid handles, just to make unit testing easier.
diff --git a/Engine/Plugins/Experimental/AsyncMessageSystem/Source/AsyncMessageSystem/Private/Tests/AsyncMessageSystemTests.cpp b/Engine/Plugins/Experimental/AsyncMessageSystem/Source/AsyncMessageSystem/Private/Tests/AsyncMessageSystemTests.cpp
--- a/Engine/Plugins/Experimental/AsyncMessageSystem/Source/AsyncMessageSystem/Private/Tests/AsyncMessageSystemTests.cpp
+++ b/Engine/Plugins/Experimental/AsyncMessageSystem/Source/AsyncMessageSystem/Private/Tests/AsyncMessageSystemTests.cpp
@@ -542,6 +542,53 @@ namespace UE::AsyncMessageSystem
 		return true;
 	}
 
+	// This test ensures that if you: 
+	// 1. Bind a listener
+	// 2. Call "Unbind" on it immediately, before the subsystem can tick
+	// 3. An event is queued and fired for the same message you were bound it
+	// 4. The listener from 1 does NOT get fired.
+	IMPLEMENT_SIMPLE_AUTOMATION_TEST(TMessageSystemUnBindingBeforeBindingFinishes, "AsyncMessagePassing.UnBinding.BeforeBindingFinishes", QuickTestFlags)
+	bool TMessageSystemUnBindingBeforeBindingFinishes::RunTest(const FString& InParameters)
+	{
+		// Startup a new message system...
+		TSharedPtr<FTestMessageSystem, ESPMode::ThreadSafe> TestSystem = FAsyncMessageSystemBase::CreateMessageSystem<FTestMessageSystem>();
+		TestTrue(TEXT("MessageSystem startup has been called"), TestSystem->HasStarted());
+
+		// Set up a listener which would modify some test data if called.
+		int32 TestData = 0;
+		const FAsyncMessageHandle ListenerHandle = TestSystem->BindListener(TestMessageId_A, [&TestData, this](const FAsyncMessage& Message)
+		{
+			TestData++;
+			TestTrue(TEXT("Test failed. This delegate should never get called."), false);
+		});
+
+		// Note: When you the call Bind/Unbind listener functions, the handles are put a queue, so they will not immediately be made valid or invalid.
+
+		// Ensure that the pre-tick state is correct. The handle shouldn't be bound yet because we have not ticked the subsystem
+		TestFalse(TEXT("The handle is in the expected pre-tick state of being unbound"), TestSystem->IsHandleBound(ListenerHandle));
+		
+		// Unbind the listener now, before we tick and ProcessListenersPendingBinding gets called.
+		TestSystem->UnbindListener(ListenerHandle);
+
+		// Tick the subsytem, which should process the pending queue of listeners and bind the delegates internally.
+		TestSystem->RunOnce();
+
+		// Since we called Unbind immediately, the handle should never be fully bound.
+		TestFalse(TEXT("The handle was correctly not bound during tick."), TestSystem->IsHandleBound(ListenerHandle));
+
+		// Fire a message for the event which was bound/unbound. The test data should not be mutated
+		TestSystem->QueueMessageForBroadcast(TestMessageId_A);
+		TestSystem->RunOnce();
+
+		// Ensure that our test data has not been modified after firing a delegate
+		TestTrue(TEXT("The handle is in the expected post-tick state and is unbound"), TestData == 0);
+
+		// All done
+		TestSystem->Shutdown();
+		TestSystem.Reset();
+		return true;
+	}
+	
 	// Test that the handle of an message is successfully unbound
 	IMPLEMENT_SIMPLE_AUTOMATION_TEST(TMessageSystemUnBindingMultiple, "AsyncMessagePassing.UnBinding.Multiple", QuickTestFlags)
 	bool TMessageSystemUnBindingMultiple::RunTest(const FString& InParameters)
diff --git a/Engine/Plugins/Experimental/AsyncMessageSystem/Source/AsyncMessageSystem/Public/AsyncMessageSystemBase.h b/Engine/Plugins/Experimental/AsyncMessageSystem/Source/AsyncMessageSystem/Public/AsyncMessageSystemBase.h
--- a/Engine/Plugins/Experimental/AsyncMessageSystem/Source/AsyncMessageSystem/Public/AsyncMessageSystemBase.h
+++ b/Engine/Plugins/Experimental/AsyncMessageSystem/Source/AsyncMessageSystem/Public/AsyncMessageSystemBase.h
@@ -311,6 +311,17 @@ protected:
 	 */
 	TMpscQueue<FAsyncMessageHandle> MessageHandlesPendingRemoval;		
 
+	/**
+	* A set of handles which have been marked for removal, but their associated listeners have not been
+	* fully bound yet (their binding endpoint does not have any data about them yet). 
+	* 
+	* Handles in this set will be skipped if they are seen during ProcessListenersPendingBinding, as they 
+	* have already been marked for removal.
+	* 
+	* This should only be utilized behind the MessageListenerMapCS scope lock
+	*/
+	TSet<FAsyncMessageHandle> IncompleteBoundHandlesToRemove;
+
 #if WITH_DEV_AUTOMATION_TESTS
 public:
     /**
