[Async Message System] Use the new TickTaskManager to process events bound in TickGroups.

This has two affects for listeners with binding options set to use Tick Groups.

1. We can be more efficient and only execute the message system's TickFunction (FMessageSystemTickFunction) *when there are pending messages in the queue for that tick group*.
2. When you queue a message during a tick group, if there are listeners for that message within that same group, they will now get called at the end of the tick group instead of having a 1 frame delay.

Also, removed some unnecessary API exposures on the FAsyncGameplayMessageSystem type, as they were private and nothing could call them outside this module anyway.

Note: this is currently behind the AsyncMessageSystem.UseTaskSyncManager CVar, which can be set to false if there are bugs with the task sync manager implementation. This CVar will be removed in a future release to always favor the more performant task sync man.

#rnx

#Synced-CL 48656164

diff --git a/Engine/Plugins/Experimental/AsyncMessageSystem/Source/AsyncMessageSystem/Private/AsyncGameplayMessageSystem.cpp b/Engine/Plugins/Experimental/AsyncMessageSystem/Source/AsyncMessageSystem/Private/AsyncGameplayMessageSystem.cpp
--- a/Engine/Plugins/Experimental/AsyncMessageSystem/Source/AsyncMessageSystem/Private/AsyncGameplayMessageSystem.cpp
+++ b/Engine/Plugins/Experimental/AsyncMessageSystem/Source/AsyncMessageSystem/Private/AsyncGameplayMessageSystem.cpp
@@ -6,6 +6,7 @@
 #include "AsyncMessageSystemLogs.h"
 #include "Engine/Level.h"
 #include "Engine/World.h"
+#include "HAL/IConsoleManager.h"
 #include "Tasks/TaskPrivate.h"
 
 namespace UE::Private
@@ -17,6 +18,24 @@ namespace UE::Private
 		
 		return TGEnum ? TGEnum->GetDisplayNameTextByValue(Group).ToString() : Invalid;
 	};
+
+	static bool bUseTaskSyncManager = false;
+	static FAutoConsoleVariableRef CVarUseTaskSyncManager(TEXT("AsyncMessageSystem.UseTaskSyncManager"),
+		bUseTaskSyncManager,
+		TEXT("If true, the TaskSyncManager will be used to control when tick group messages will be processed."));
+
+	/** The description to use when creating sync points with the task sync manager */
+	static const FName AsyncMessageSyncPointDesc = TEXT("AsyncGameplayMessageSystem_SyncPointDesc");
+
+	static const TMap<ETickingGroup, FName> SyncPointTickGroupNames =
+	{
+		{ TG_PrePhysics, 		"TG_PrePhysics" 	},
+		{ TG_StartPhysics, 		"TG_StartPhysics" 	},
+		{ TG_DuringPhysics, 	"TG_DuringPhysics"	},
+		{ TG_EndPhysics, 		"TG_EndPhysics" 	},
+		{ TG_PostPhysics, 		"TG_PostPhysics" 	},
+		{ TG_PostUpdateWork, 	"TG_PostUpdateWork" },
+	};
 }
 
 /**
@@ -31,7 +50,7 @@ struct FMessageSystemTickFunction final : public FTickFunction
 		, WeakMessageSys(InWeakMessageSystem)
 	{
 		bCanEverTick = true;
-		bStartWithTickEnabled = true;
+		bStartWithTickEnabled = !UE::Private::bUseTaskSyncManager;
 		bAllowTickBatching = true;
 
 		// Only run this tick function on the game thread, because our message system is our "sync point"
@@ -49,7 +68,7 @@ struct FMessageSystemTickFunction final : public FTickFunction
 		// Each tick function can simply call the message system and let it know that the next tick group has started.
 		if (TSharedPtr<FAsyncGameplayMessageSystem> MessageSys = WeakMessageSys.Pin())
 		{
-			MessageSys->ExecuteTick(DeltaTime, TickType, CurrentThread, MyCompletionGraphEvent, TickGroup);
+			MessageSys->ProcessMessagesForTickGroup(TickGroup);
 		}
 	}
 	
@@ -92,14 +111,16 @@ void FAsyncGameplayMessageSystem::PostQueueMessage(const FAsyncMessageId Message
 	// async task to process their message queue. 
 	for (const FAsyncMessageBindingOptions& BindingOpts : OptionsBoundTo)
 	{
-		if (BindingOpts.GetType() == FAsyncMessageBindingOptions::EBindingType::UseNamedThreads || BindingOpts.GetType() == FAsyncMessageBindingOptions::EBindingType::UseTaskPriorities)
+		switch (BindingOpts.GetType())
 		{
+		case FAsyncMessageBindingOptions::EBindingType::UseTickGroup:
+			EnableExecutionForTickGroupBinding(BindingOpts);
+			break;
+		case FAsyncMessageBindingOptions::EBindingType::UseNamedThreads:
+		case FAsyncMessageBindingOptions::EBindingType::UseTaskPriorities:
 			StartAsyncProcessForBinding(BindingOpts);
+			break;
 		}
-
-		// TODO: We could potentially use this to conditionally control when the tick groups need to run as well, reducing the amount that they are called.
-		// That way, the tick functions only execute if there are currently messages in their queue. We could possibly do this by registering/unregistering the
-		// tick functions as needed for the tick groups as they run. For now though, just let them tick
 	}
 }
 
@@ -112,6 +133,8 @@ void FAsyncGameplayMessageSystem::CreateTickFunctions()
 		return;
 	}
 
+	FScopeLock WorkHandleLock(&WorkHandleCreationCS);
+
 	// We will be binding to some owning world's persistent level to create our tick functions
 	ULevel* TickLevel = OuterWorld->PersistentLevel;
 
@@ -123,6 +146,9 @@ void FAsyncGameplayMessageSystem::CreateTickFunctions()
 	const int StartingGroup = static_cast<int>(EarliestSupportedTickGroup);
 	const int LastGroup = static_cast<int>(LatestSupportedTickGroup);
 
+	using namespace UE::Tick;
+	FTaskSyncManager* SyncMan = UE::Private::bUseTaskSyncManager ? FTaskSyncManager::Get() : nullptr;
+
 	// Spawn a tick function for every tick group that we can actually do any work in
 	for (int i = StartingGroup; i <= LastGroup; ++i)
 	{
@@ -130,32 +156,69 @@ void FAsyncGameplayMessageSystem::CreateTickFunctions()
 
 		TSharedPtr<FTickFunction> Func = MakeShared<FMessageSystemTickFunction>(Group, WeakThisPtr);
 
-		Func->RegisterTickFunction(TickLevel);
-
-		// We always want the previous tick function to wait until the next one to start processing
-		if (PreviousTickFunction.IsValid())
+		if (!SyncMan)
+		{
+			// Fall back to the "normal" tick functoin process of ticking every frame
+			// if we are not using the TaskSyncManager
+			Func->RegisterTickFunction(TickLevel);
+
+			// We always want the previous tick function to wait until the next one to start processing
+			if (PreviousTickFunction.IsValid())
+			{
+				Func->AddPrerequisite(TickLevel, *PreviousTickFunction);
+			}	
+		}
+		else
 		{
-			Func->AddPrerequisite(TickLevel, *PreviousTickFunction);
+			const FName* GroupSyncPointName = UE::Private::SyncPointTickGroupNames.Find(Group);
+			check(GroupSyncPointName);
+			
+			FSyncPointDescription Desc = {};
+			Desc.FirstPossibleTickGroup = Group;
+			Desc.LastPossibleTickGroup = Group;
+			Desc.EventType = ESyncPointEventType::GameThreadTask;
+			Desc.ActivationRules = ESyncPointActivationRules::WaitForAllWork;
+			Desc.SourceName = UE::Private::AsyncMessageSyncPointDesc;
+			Desc.RegisteredName = *GroupSyncPointName;
+			
+			ensure(SyncMan->RegisterNewSyncPoint(Desc));
 		}
 
 		PreviousTickFunction = Func;
 
 		// Keep track of the tick functions we have created so that we can properly unregister them later
-		TickFunctions.Emplace(MoveTemp(Func));
+		TickFunctions.Add(Group, MoveTemp(Func));
 	}
 }
 
 void FAsyncGameplayMessageSystem::DestroyTickFunctions()
 {
-	for (TSharedPtr<FTickFunction> Func : TickFunctions)
+	using namespace UE::Tick;
+
+	FScopeLock WorkHandleLock(&WorkHandleCreationCS);
+	
+	FTaskSyncManager* SyncMan = UE::Private::bUseTaskSyncManager ? FTaskSyncManager::Get() : nullptr;
+	
+	for (TPair<ETickingGroup, TSharedPtr<FTickFunction>>& Pair : TickFunctions)
 	{
-		if (Func.IsValid())
+		if (Pair.Value.IsValid())
 		{
-			Func->UnRegisterTickFunction();
+			Pair.Value->UnRegisterTickFunction();
+		}
+		
+		if (SyncMan)
+		{
+			const FName* GroupSyncPointName = UE::Private::SyncPointTickGroupNames.Find(Pair.Key);
+			check(GroupSyncPointName);
+			
+			ensure(SyncMan->UnregisterSyncPoint(*GroupSyncPointName, UE::Private::AsyncMessageSyncPointDesc));
 		}
 	}
 
 	TickFunctions.Empty();
+
+	// The work handle will abandon work if necessary here in its destructor, so we can just empty the map.
+	WorkHandles.Empty();
 }
 
 void FAsyncGameplayMessageSystem::StartAsyncProcessForBinding(const FAsyncMessageBindingOptions& Options)
@@ -190,28 +253,97 @@ void FAsyncGameplayMessageSystem::StartAsyncProcessForBinding(const FAsyncMessag
     },
     TaskPri,
     ExtendedTaskPri);
+}
+
+void FAsyncGameplayMessageSystem::EnableExecutionForTickGroupBinding(const FAsyncMessageBindingOptions& Options)
+{
+	check(Options.GetType() == FAsyncMessageBindingOptions::EBindingType::UseTickGroup);
+
+	if (!UE::Private::bUseTaskSyncManager)
+	{
+		return;
+	}
+
+	FScopeLock WorkHandleLock(&WorkHandleCreationCS);
+
+	const ETickingGroup Group = Options.GetTickGroup();
+	
+	UE::Tick::FActiveSyncWorkHandle* WorkHandle = GetOrCreateWorkHandleForTickGroup(Group);
+	if (!WorkHandle)
+	{
+		UE_LOG(LogAsyncMessageSystem, Error, TEXT("Failed to get a work handle for group '%s'. Messages will not be processed."),
+			*UE::Private::LexToString(Options.GetTickGroup()));
+		return;
+	}
+
+	TSharedPtr<FTickFunction>* TickFunc = TickFunctions.Find(Group);
+	if (!ensure(TickFunc))
+	{
+		UE_LOG(LogAsyncMessageSystem, Error, TEXT("Failed to get a tick function '%s'. Messages will not be processed."),
+			*UE::Private::LexToString(Options.GetTickGroup()));
+		return;
+	}
 	
-	// The above UE::Tasks::Launch is the same thing as this older Task Graph syntax:
-	// AsyncTask(ThreadToProcessOn, [WeakThisPtr, Options]()
-	// {
-	// 	if (TSharedPtr<FAsyncGameplayMessageSystem> This = WeakThisPtr.Pin())
-	// 	{
-	// 		This->ProcessMessageQueueForBinding(Options);
-	// 	}
-	// });
-	//
-	// We utilize the UE::Tasks system here and translate the ENamedThread because the UE::Tasks system
-	// should have less overhead and better scheduling capabilities in most scenarios
+	const bool bSuccess = WorkHandle->RequestWork(TickFunc->Get(), UE::Tick::ESyncWorkRepetition::Once);
+
+	UE_CLOG(!bSuccess, LogAsyncMessageSystem, Error, TEXT("Failed to request work for group '%s'. Messages will not be processed."),
+		*UE::Private::LexToString(Options.GetTickGroup()));
 }
 
-void FAsyncGameplayMessageSystem::ExecuteTick(
-	float DeltaTime,
-	ELevelTick TickType,
-	ENamedThreads::Type CurrentThread,
-	const FGraphEventRef& MyCompletionGraphEvent,
-	ETickingGroup TickGroup)
+UE::Tick::FActiveSyncWorkHandle* FAsyncGameplayMessageSystem::GetOrCreateWorkHandleForTickGroup(const ETickingGroup Group)
 {
-	TRACE_CPUPROFILER_EVENT_SCOPE(FMessageSystem::ExecuteTick);
+	if (!ensure(UE::Private::bUseTaskSyncManager))
+	{
+		return nullptr;
+	}
+	
+	//FScopeLock WorkHandleLock(&WorkHandleCreationCS);
+
+	// First, check if we have a handle already.
+	if (UE::Tick::FActiveSyncWorkHandle* ExistingHandle = WorkHandles.Find(Group))
+	{
+		return ExistingHandle;
+	}
+
+	UWorld* World = OuterWorld.Get();
+	if (!World)
+	{
+		UE_LOG(LogAsyncMessageSystem, Error, TEXT("Invalid world, cannot create work handles."));
+		return nullptr;
+	}
+	
+	using namespace UE::Tick;
+	FTaskSyncManager* SyncMan = FTaskSyncManager::Get();
+	if (!SyncMan)
+	{
+		UE_LOG(LogAsyncMessageSystem, Error, TEXT("Unable to get a valid FTaskSyncManager. Messages may not be processed."));
+		return nullptr;
+	}
+	const FName* GroupSyncPointName = UE::Private::SyncPointTickGroupNames.Find(Group);
+	check(GroupSyncPointName);
+	
+	FSyncPointId SyncPoint = SyncMan->FindSyncPoint(World, *GroupSyncPointName);
+	if (!SyncPoint.IsValid())
+	{
+		UE_LOG(LogAsyncMessageSystem, Error, TEXT("Failed to get a sync point for group '%s'. Messages will not be processed."),
+			*UE::Private::LexToString(Group));
+		return nullptr;
+	}
+
+	FActiveSyncWorkHandle WorkHandle;
+	FTaskSyncResult Res = SyncMan->RegisterWorkHandle(SyncPoint, WorkHandle);
+	if (!ensure(Res.WasSuccessful()))
+	{
+		UE_LOG(LogAsyncMessageSystem, Error, TEXT("Failed to register work handle for group '%s'. Messages will not be processed."),
+			*UE::Private::LexToString(Group));
+	}
+
+	return &WorkHandles.Add(Group, MoveTemp(WorkHandle));
+}
+
+void FAsyncGameplayMessageSystem::ProcessMessagesForTickGroup(const ETickingGroup TickGroup)
+{
+	TRACE_CPUPROFILER_EVENT_SCOPE(FMessageSystem::ProcessMessagesForTickGroup);
 
 	if (!OuterWorld.IsValid())
 	{
@@ -222,18 +354,6 @@ void FAsyncGameplayMessageSystem::ExecuteTick(
 	
 	CurrentTickGroup = TickGroup;
 
-	// TODO: Right now the message system uses GFrameCounter as a quick and easy way to determine when messages are sent
-	// We should get away from that, and instead use an atomic in in the message system itself and increment it once per frame
-	
-	// If we are in the last tick group of the frame, then increment the frame counter of the message system that it is using.
-	//const bool bIsLastGroupOfFrame = (CurrentTickGroup == LatestSupportedTickGroup);
-	const uint64 CurrentFrame = GFrameCounter;
-	
-	UE_LOG(LogAsyncMessageSystem, VeryVerbose, TEXT("Frame %llu Exectute tick %s :: Last tick group was %s"),
-		CurrentFrame,
-		*UE::Private::LexToString(CurrentTickGroup),
-		*UE::Private::LexToString(LastTickedGroup));
-
 	ensure(CurrentTickGroup >= EarliestSupportedTickGroup || CurrentTickGroup <= LatestSupportedTickGroup);	
 
 	// Process the messages for this current tick group
@@ -243,4 +363,4 @@ void FAsyncGameplayMessageSystem::ExecuteTick(
 	ProcessMessageQueueForBinding(Options);
 		
 	LastTickedGroup = TickGroup;
-}
\ No newline at end of file
+}
diff --git a/Engine/Plugins/Experimental/AsyncMessageSystem/Source/AsyncMessageSystem/Public/AsyncGameplayMessageSystem.h b/Engine/Plugins/Experimental/AsyncMessageSystem/Source/AsyncMessageSystem/Public/AsyncGameplayMessageSystem.h
--- a/Engine/Plugins/Experimental/AsyncMessageSystem/Source/AsyncMessageSystem/Public/AsyncGameplayMessageSystem.h
+++ b/Engine/Plugins/Experimental/AsyncMessageSystem/Source/AsyncMessageSystem/Public/AsyncGameplayMessageSystem.h
@@ -5,6 +5,7 @@
 #include "Async/TaskGraphFwd.h"			// FGraphEventRef
 #include "AsyncMessageSystemBase.h"
 #include "Engine/EngineBaseTypes.h"		// For ELevelTick
+#include "TaskSyncManager.h"
 
 #define UE_API ASYNCMESSAGESYSTEM_API
 
@@ -24,20 +25,20 @@ public:
 
 private:
 	//~ Begin FAsyncMessageSystemBase interface
-	UE_API virtual void Startup_Impl() override;
-	UE_API virtual void Shutdown_Impl() override;
-	UE_API virtual void PostQueueMessage(const FAsyncMessageId MessageId, const TArray<FAsyncMessageBindingOptions>& OptionsBoundTo) override;
+	virtual void Startup_Impl() override;
+	virtual void Shutdown_Impl() override;
+	virtual void PostQueueMessage(const FAsyncMessageId MessageId, const TArray<FAsyncMessageBindingOptions>& OptionsBoundTo) override;
 	//~ End FAsyncMessageSystemBase interface
 
 	/**
 	 * Creates and registers tick functions for this message system to the outer world.
 	 */
-	UE_API void CreateTickFunctions();
+	void CreateTickFunctions();
 
 	/**
 	 * Destroy any registered tick functions on this message system
 	 */
-	UE_API void DestroyTickFunctions();
+	void DestroyTickFunctions();
 
 	/**
 	 * Starts an async task to process the messages with listeners that have the given binding options.
@@ -45,17 +46,31 @@ private:
 	 * 
 	 * @param Options The options of which to start async processing for.
 	 */
-	UE_API void StartAsyncProcessForBinding(const FAsyncMessageBindingOptions& Options);
+	void StartAsyncProcessForBinding(const FAsyncMessageBindingOptions& Options);
+
+	/**
+	 * Called when a message is queued with listeners who have bindings options set to use a Tick Group.
+	 *
+	 * This will request a tick function to execute once via the TaskSyncManager to process the message queue for the given group.
+	 * 
+	 * @param Options The binding options that we want to enable the tick function for to process its listeners
+	 */
+	void EnableExecutionForTickGroupBinding(const FAsyncMessageBindingOptions& Options);
+
+	/**
+	 * Get or create the work handle required to request work from the task sync manager.
+	 * 
+	 * @param Group The tick group to get the sync work handle for
+	 * @return Pointer to the work handle which can then be used to request work. Nullptr if the work handle cannot be created. 
+	 */
+	UE::Tick::FActiveSyncWorkHandle* GetOrCreateWorkHandleForTickGroup(const ETickingGroup Group);
 	
 	/**
-	 * Executes a single tick of this message system for a given group 
+	 * Processes all pending messages in the queue for the given tick group. This is meant to be called from a FTickFunction.
+	 *
+	 * This will just create the necessary FAsyncMessageBindingOptions args to call the FAsyncMessageSystemBase::ProcessMessageQueueForBinding.
 	 */
-	UE_API void ExecuteTick(
-		float DeltaTime,
-		ELevelTick TickType,
-		ENamedThreads::Type CurrentThread,
-		const FGraphEventRef& MyCompletionGraphEvent,
-		ETickingGroup TickGroup);
+	void ProcessMessagesForTickGroup(const ETickingGroup TickGroup);
 
 	/**
 	 * The world in which this message system belongs to. The tick functions of this message system will
@@ -68,7 +83,14 @@ private:
 	 * for each tick group that this message system supports and they will depend on one another
 	 * to keep a deterministic ordering of ticking on this message system.
 	 */
-	TArray<TSharedPtr<FTickFunction>> TickFunctions;
+	TMap<ETickingGroup, TSharedPtr<FTickFunction>> TickFunctions;
+
+	/**
+	 * Handles that we can use to request work (in our case, process messages for a specific tick group)
+	 *
+	 * These handles need to be accessed behind the WorkHandleCreationCS
+	 */
+	TMap<ETickingGroup, UE::Tick::FActiveSyncWorkHandle> WorkHandles;
 
 	/***
 	 * The tick group earliest in the frame which this message system supports.
@@ -85,6 +107,9 @@ private:
 	 */
 	ETickingGroup CurrentTickGroup = EarliestSupportedTickGroup;
 	ETickingGroup LastTickedGroup = LatestSupportedTickGroup;
+	
+	/**  */
+	FCriticalSection WorkHandleCreationCS;
 };
 
 #undef UE_API
diff --git a/Engine/Plugins/Experimental/AsyncMessageSystemTests/Source/AsyncMessageSystemTests/Private/AsyncMessageSystemPerformanceTests.cpp b/Engine/Plugins/Experimental/AsyncMessageSystemTests/Source/AsyncMessageSystemTests/Private/AsyncMessageSystemPerformanceTests.cpp
--- a/Engine/Plugins/Experimental/AsyncMessageSystemTests/Source/AsyncMessageSystemTests/Private/AsyncMessageSystemPerformanceTests.cpp
+++ b/Engine/Plugins/Experimental/AsyncMessageSystemTests/Source/AsyncMessageSystemTests/Private/AsyncMessageSystemPerformanceTests.cpp
@@ -162,7 +162,7 @@ namespace UE::AsyncMessageSystem
 			DestroyAllTestActors();
 
 			// Store how many actors we are using
-			ActorCount = ActorCountOverride > 0 ? ActorCountOverride : CVarAsyncMessageSystemTestTickCount.GetValueOnAnyThread();
+			ActorCount = ActorCountOverride > 0 ? ActorCountOverride : CVarAsyncMessageSystemTestActorCount.GetValueOnAnyThread();
 
 			// Create the test actors
 			bSuccess &= CreateTestActors(ActorCount, AASyncMessagePerfTest::StaticClass());
@@ -192,7 +192,15 @@ namespace UE::AsyncMessageSystem
 			}
 
 			// Lastly, begin play in the test world
-			bSuccess &= BeginPlayInTestWorld();			
+			bSuccess &= BeginPlayInTestWorld();
+
+			// and tick the world once without any tests running
+			// Tick the world once before the test ticks to give any underlying systems a chance to finish initaliziating, like the task sync manager.
+			if (true)
+			{
+				TickTestWorld(SimulatedDeltaTime);
+				TickCount++;
+			}
 
 			return bSuccess;
 		}
