[Async Message System] Remove overzealous ensures about the task sync manager.

Note: there is currently an issue when using AsyncMessageSystem.UseTaskSyncManager=1 during level transitions where these ensures would trip as well.  That bug is tracked with UE-357378

#rnx
#Synced-CL 49065031

diff --git a/Engine/Plugins/Experimental/AsyncMessageSystem/Source/AsyncMessageSystem/Private/AsyncGameplayMessageSystem.cpp b/Engine/Plugins/Experimental/AsyncMessageSystem/Source/AsyncMessageSystem/Private/AsyncGameplayMessageSystem.cpp
--- a/Engine/Plugins/Experimental/AsyncMessageSystem/Source/AsyncMessageSystem/Private/AsyncGameplayMessageSystem.cpp
+++ b/Engine/Plugins/Experimental/AsyncMessageSystem/Source/AsyncMessageSystem/Private/AsyncGameplayMessageSystem.cpp
@@ -181,7 +181,7 @@ void FAsyncGameplayMessageSystem::CreateTickFunctions()
 			Desc.SourceName = UE::Private::AsyncMessageSyncPointDesc;
 			Desc.RegisteredName = *GroupSyncPointName;
 			
-			ensure(SyncMan->RegisterNewSyncPoint(Desc));
+			SyncMan->RegisterNewSyncPoint(Desc);
 		}
 
 		PreviousTickFunction = Func;
@@ -211,7 +211,7 @@ void FAsyncGameplayMessageSystem::DestroyTickFunctions()
 			const FName* GroupSyncPointName = UE::Private::SyncPointTickGroupNames.Find(Pair.Key);
 			check(GroupSyncPointName);
 			
-			ensure(SyncMan->UnregisterSyncPoint(*GroupSyncPointName, UE::Private::AsyncMessageSyncPointDesc));
+			SyncMan->UnregisterSyncPoint(*GroupSyncPointName, UE::Private::AsyncMessageSyncPointDesc);
 		}
 	}
 
