[Async Message System] In UAsyncMessageBindingComponent, create the endpoint on Initialize component instead of begin play so that it is more readily available for use in blueprints.

#rnx
#Synced-CL 49065031

diff --git a/Engine/Plugins/Experimental/AsyncMessageSystem/Source/AsyncMessageSystem/Private/AsyncMessageBindingComponent.cpp b/Engine/Plugins/Experimental/AsyncMessageSystem/Source/AsyncMessageSystem/Private/AsyncMessageBindingComponent.cpp
--- a/Engine/Plugins/Experimental/AsyncMessageSystem/Source/AsyncMessageSystem/Private/AsyncMessageBindingComponent.cpp
+++ b/Engine/Plugins/Experimental/AsyncMessageSystem/Source/AsyncMessageSystem/Private/AsyncMessageBindingComponent.cpp
@@ -5,16 +5,16 @@
 
 #include UE_INLINE_GENERATED_CPP_BY_NAME(AsyncMessageBindingComponent)
 
-void UAsyncMessageBindingComponent::BeginPlay()
+void UAsyncMessageBindingComponent::InitializeComponent()
 {
-	Super::BeginPlay();
+	Super::InitializeComponent();
 
 	CreateEndpoint();
 }
 
-void UAsyncMessageBindingComponent::EndPlay(const EEndPlayReason::Type EndPlayReason)
+void UAsyncMessageBindingComponent::UninitializeComponent()
 {
-	Super::EndPlay(EndPlayReason);
+	Super::UninitializeComponent();
 
 	CleanupEndpoint();
 }
diff --git a/Engine/Plugins/Experimental/AsyncMessageSystem/Source/AsyncMessageSystem/Public/AsyncMessageBindingComponent.h b/Engine/Plugins/Experimental/AsyncMessageSystem/Source/AsyncMessageSystem/Public/AsyncMessageBindingComponent.h
--- a/Engine/Plugins/Experimental/AsyncMessageSystem/Source/AsyncMessageSystem/Public/AsyncMessageBindingComponent.h
+++ b/Engine/Plugins/Experimental/AsyncMessageSystem/Source/AsyncMessageSystem/Public/AsyncMessageBindingComponent.h
@@ -39,7 +39,8 @@ public:
  * A blueprint component which will allow you to specify a specific endpoint to use
  * when queueing messages or binding listeners to the Async Message System.
  *
- * By default, this will create an endpoint on BeginPlay and destroy it on EndPlay.
+ * By default, this will create an endpoint on InitializeComponent (making it ready for use during Begin Play) 
+ * and destroy it on UninitializeComponent.
  */
 UCLASS(MinimalAPI, BlueprintType, HideCategories=(Object,Replication,Navigation), Meta=(BlueprintSpawnableComponent))
 class UAsyncMessageBindingComponent :
@@ -51,8 +52,8 @@ class UAsyncMessageBindingComponent :
 protected:
 
 	//~ Begin UActorComponent interface
-	virtual void BeginPlay() override;
-	virtual void EndPlay(const EEndPlayReason::Type EndPlayReason) override;
+	virtual void InitializeComponent() override;
+	virtual void UninitializeComponent() override;
 	//~ end UActorComponent interface
 
 	//~ Begin IAsyncMessageBindingEndpointInterface
